StateGraph - State machine framework plugin modeled as a directed graph that supports timeouts, hotfixing, runtime modification, and more.

diff --git a/Engine/Plugins/Experimental/StateGraph/Source/StateGraph/Private/StateGraph.cpp b/Engine/Plugins/Experimental/StateGraph/Source/StateGraph/Private/StateGraph.cpp
new file mode 100644
--- /dev/null
+++ b/Engine/Plugins/Experimental/StateGraph/Source/StateGraph/Private/StateGraph.cpp
@@ -0,0 +1,743 @@
+// Copyright Epic Games, Inc. All Rights Reserved.
+
+#include "StateGraph.h"
+
+#include "HAL/PlatformTime.h"
+#include "Misc/ConfigCacheIni.h"
+#include "Misc/Parse.h"
+#include "Modules/ModuleManager.h"
+
+DEFINE_LOG_CATEGORY(LogStateGraph);
+
+class FStateGraphModule : public IModuleInterface
+{
+};
+
+IMPLEMENT_MODULE(FStateGraphModule, StateGraph);
+
+namespace UE {
+
+// FStateGraphNode
+
+FStateGraphNode::FStateGraphNode(FName InName) :
+	Name(InName)
+{
+	UE_LOG_STATEGRAPH(VeryVerbose, TEXT("[%s.%s] Created node"), *GetStateGraphName().ToString(), *Name.ToString());
+}
+
+FStateGraphNode::~FStateGraphNode()
+{
+	UE_LOG_STATEGRAPH(VeryVerbose, TEXT("[%s.%s] Destroyed node"), *GetStateGraphName().ToString(), *Name.ToString());
+}
+
+const TCHAR* FStateGraphNode::GetStatusName(EStatus Status)
+{
+	switch (Status)
+	{
+		case EStatus::NotStarted: return TEXT("NotStarted");
+		case EStatus::Blocked: return TEXT("Blocked");
+		case EStatus::Started: return TEXT("Started");
+		case EStatus::Completed: return TEXT("Completed");
+		case EStatus::TimedOut: return TEXT("TimedOut");
+		default: checkNoEntry(); return TEXT("Unknown");
+	}
+}
+
+void FStateGraphNode::SetTimeout(double InTimeout)
+{
+	Timeout = InTimeout;
+}
+
+double FStateGraphNode::GetDuration() const
+{
+	if (StartTime == 0.f)
+	{
+		return 0.f;
+	}
+
+	return (CompletedTime == 0.f ? FPlatformTime::Seconds() : CompletedTime) - StartTime;
+}
+
+bool FStateGraphNode::CheckDependencies() const
+{
+	FStateGraphPtr StateGraphPtr(StateGraphWeakPtr.Pin());
+	if (!StateGraphPtr.IsValid())
+	{
+		UE_LOG_STATEGRAPH(Warning, TEXT("[%s] Node checked with invalid state graph"), *Name.ToString());
+		return false;
+	}
+
+	for (FName Dependency : Dependencies)
+	{
+		FStateGraphNodeRef* DependencyNode = StateGraphPtr->GetNodeRef(Dependency);
+		if (!DependencyNode || (*DependencyNode)->GetStatus() != FStateGraphNode::EStatus::Completed)
+		{
+			return false;
+		}
+	}
+
+	return true;
+}
+
+void FStateGraphNode::Complete()
+{
+	if (Status == EStatus::Completed)
+	{
+		UE_LOG_STATEGRAPH(Warning, TEXT("[%s.%s] Node already completed"), *GetStateGraphName().ToString(), *Name.ToString());
+		return;
+	}
+
+	CompletedTime = FPlatformTime::Seconds();
+	UE_LOG_STATEGRAPH(Log, TEXT("[%s.%s] Completed node (Duration=%.6f)"), *GetStateGraphName().ToString(), *Name.ToString(), CompletedTime - StartTime);
+
+	// Keep a reference to check if the node is destroyed during external functions.
+	FStateGraphNodeWeakPtr StateGraphNodeWeakPtr(AsWeak());
+
+	SetStatus(EStatus::Completed);
+
+	if (!StateGraphNodeWeakPtr.IsValid())
+	{
+		return;
+	}
+
+	// Run state graph again since this completion may have fulfilled dependencies to start new nodes.
+	FStateGraphPtr StateGraphPtr = StateGraphWeakPtr.Pin();
+	if (StateGraphPtr.IsValid())
+	{
+		if (StateGraphPtr->GetStatus() != FStateGraph::EStatus::Paused)
+		{
+			StateGraphPtr->Run();
+		}
+	}
+	else
+	{
+		UE_LOG_STATEGRAPH(Warning, TEXT("[%s] Node completed with invalid state graph"), *Name.ToString());
+	}
+}
+
+void FStateGraphNode::Reset()
+{
+	UE_LOG_STATEGRAPH(Verbose, TEXT("[%s.%s] Resetting node"), *GetStateGraphName().ToString(), *Name.ToString());
+	StartTime = 0.f;
+	CompletedTime = 0.f;
+	SetStatus(EStatus::NotStarted);
+}
+
+FName FStateGraphNode::GetStateGraphName() const
+{
+	return StateGraphWeakPtr.IsValid() ? StateGraphWeakPtr.Pin()->GetName() : NAME_None;
+}
+
+void FStateGraphNode::UpdateConfig()
+{
+	GConfig->GetDouble(*ConfigSectionName, TEXT("Timeout"), Timeout, GEngineIni);
+}
+
+void FStateGraphNode::SetStatus(EStatus NewStatus)
+{
+	if (Status != NewStatus)
+	{
+		const EStatus OldStatus = Status;
+		Status = NewStatus;
+		FStateGraphPtr StateGraphPtr = StateGraphWeakPtr.Pin();
+		if (StateGraphPtr.IsValid())
+		{
+			StateGraphPtr->OnNodeStatusChanged.Broadcast(*this, OldStatus, NewStatus);
+		}
+	}
+}
+
+// FStateGraphNodeFunction
+
+FStateGraphNodeFunction::FStateGraphNodeFunction(FName InName, const FStateGraphNodeFunctionStart& InStartFunction) :
+	FStateGraphNode(InName),
+	StartFunction(InStartFunction)
+{}
+
+bool FStateGraphNodeFunction::CheckDependencies() const
+{
+	if (!StartFunction.IsBound())
+	{
+		UE_LOG_STATEGRAPH(Warning, TEXT("[%s.%s] Function node start not bound"), *GetStateGraphName().ToString(), *GetName().ToString());
+		return false;
+	}
+
+	return FStateGraphNode::CheckDependencies();
+}
+
+void FStateGraphNodeFunction::Start()
+{
+	// CheckDependencies verified the function is bound right before calling this in StateGraph->Run().
+	check(StartFunction.IsBound());
+
+	FStateGraphNodeWeakPtr WeakPtr(AsWeak());
+	StartFunction.Execute([WeakPtr, StateGraphName = GetStateGraphName(), NodeName = GetName()]()
+	{
+		if (WeakPtr.IsValid())
+		{
+			WeakPtr.Pin()->Complete();
+		}
+		else
+		{
+			UE_LOG_STATEGRAPH(Warning, TEXT("[%s.%s] Function node completed after node was destroyed"), *StateGraphName.ToString(), *NodeName.ToString());
+		}
+	});
+}
+
+// FStateGraph
+
+FStateGraph::FStateGraph(FName InName) :
+	Name(InName),
+	ConfigSectionName(FString::Printf(TEXT("StateGraph.%s"), *InName.ToString()))
+{
+	UE_LOG_STATEGRAPH(VeryVerbose, TEXT("[%s] Created state graph"), *Name.ToString());
+}
+
+FStateGraph::~FStateGraph()
+{
+	UE_LOG_STATEGRAPH(VeryVerbose, TEXT("[%s] Destroyed state graph"), *Name.ToString());
+	FCoreDelegates::TSOnConfigSectionsChanged().Remove(ConfigSectionsChangedDelegate);
+}
+
+void FStateGraph::Initialize()
+{
+	if (!ConfigSectionsChangedDelegate.IsValid())
+	{
+		ConfigSectionsChangedDelegate = FCoreDelegates::TSOnConfigSectionsChanged().AddSP(this, &FStateGraph::OnConfigSectionsChanged);
+	}
+
+	UpdateConfig();
+}
+
+const TCHAR* FStateGraph::GetStatusName(EStatus Status)
+{
+	switch (Status)
+	{
+		case EStatus::NotStarted: return TEXT("NotStarted");
+		case EStatus::Running: return TEXT("Running");
+		case EStatus::Waiting: return TEXT("Waiting");
+		case EStatus::Blocked: return TEXT("Blocked");
+		case EStatus::Completed: return TEXT("Completed");
+		case EStatus::Paused: return TEXT("Paused");
+		case EStatus::TimedOut: return TEXT("TimedOut");
+		default: checkNoEntry(); return TEXT("Unknown");
+	}
+}
+
+void FStateGraph::SetTimeout(double InTimeout)
+{
+	Timeout = InTimeout;
+}
+
+double FStateGraph::GetDuration() const
+{
+	if (StartTime == 0.f)
+	{
+		return 0.f;
+	}
+
+	return (CompletedTime == 0.f ? FPlatformTime::Seconds() : CompletedTime) - StartTime;
+}
+
+bool FStateGraph::AddNode(const FStateGraphNodeRef& Node)
+{
+	if (Node->StateGraphWeakPtr.IsValid())
+	{
+		UE_LOG_STATEGRAPH(Warning, TEXT("[%s.%s] Node already associated with a state graph"), *Name.ToString(), *Node->GetName().ToString());
+		return false;
+	}
+
+	if (Nodes.Contains(Node->GetName()))
+	{
+		UE_LOG_STATEGRAPH(Warning, TEXT("[%s.%s] Node with the same name already exists"), *Name.ToString(), *Node->GetName().ToString());
+		return false;
+	}
+
+	UE_LOG_STATEGRAPH(Verbose, TEXT("[%s.%s] Adding node"), *Name.ToString(), *Node->GetName().ToString());
+	Node->StateGraphWeakPtr = AsWeak();
+	Node->ConfigSectionName = FString::Printf(TEXT("%s.%s"), *ConfigSectionName, *Node->GetName().ToString());
+	Nodes.Add(Node->GetName(), Node);
+	Node->UpdateConfig();
+	return true;
+}
+
+bool FStateGraph::RemoveNode(FName NodeName)
+{
+	FStateGraphNodeRef* Node = GetNodeRef(NodeName);
+	if (!Node)
+	{
+		UE_LOG_STATEGRAPH(Warning, TEXT("[%s.%s] Failed to remove node"), *Name.ToString(), *NodeName.ToString());
+		return false;
+	}
+
+	UE_LOG_STATEGRAPH(Verbose, TEXT("[%s.%s] Removing node"), *Name.ToString(), *NodeName.ToString());
+	(*Node)->StateGraphWeakPtr.Reset();
+
+	// Keep a reference to check if the state graph is destroyed during external functions.
+	FStateGraphWeakPtr StateGraphWeakPtr(AsWeak());
+	(*Node)->Removed();
+	if (StateGraphWeakPtr.IsValid())
+	{
+		Nodes.Remove(NodeName);
+	}
+
+	return true;
+}
+
+void FStateGraph::RemoveAllNodes()
+{
+	// Keep a reference to check if the state graph is destroyed during external functions.
+	FStateGraphWeakPtr StateGraphWeakPtr(AsWeak());
+
+	// Copy names since node map may be modified during loop.
+	TArray<FName> NodeNames;
+	Nodes.GenerateKeyArray(NodeNames);
+	for (FName NodeName : NodeNames)
+	{
+		RemoveNode(NodeName);
+		if (!StateGraphWeakPtr.IsValid())
+		{
+			return;
+		}
+	}
+}
+
+void FStateGraph::Run()
+{
+	const double Now = FPlatformTime::Seconds();
+
+	if (Status == EStatus::NotStarted)
+	{
+		StartTime = Now;
+	}
+
+	// Keep a reference to check if the state graph is destroyed during external functions.
+	FStateGraphWeakPtr StateGraphWeakPtr(AsWeak());
+
+	SetStatus(EStatus::Running);
+
+	if (!StateGraphWeakPtr.IsValid())
+	{
+		return;
+	}
+
+	if (bRunning)
+	{
+		bRunAgain = true;
+		return;
+	}
+
+	FTSTicker::GetCoreTicker().RemoveTicker(TimeoutTicker);
+
+	double NextTimeout = 0.f;
+	if (Timeout > 0.f)
+	{
+		NextTimeout = (StartTime + Timeout) - Now;
+		if (NextTimeout <= 0.f)
+		{
+			UE_LOG_STATEGRAPH(Log, TEXT("[%s] State graph timed out (Duration=%.6f)"), *Name.ToString(), Now - StartTime);
+			SetStatus(EStatus::TimedOut);
+			return;
+		}
+	}
+
+	bRunning = true;
+
+	UE_LOG_STATEGRAPH(Verbose, TEXT("[%s] Starting run loop"), *Name.ToString());
+
+	uint32 Blocked = 0;
+	uint32 Started = 0;
+	uint32 Running = 0;
+	uint32 Completed = 0;
+	uint32 Removed = 0;
+	uint32 TimedOut = 0;
+
+	// Copy names since node map may be modified during loop.
+	TArray<FName> NodeNames;
+	Nodes.GenerateKeyArray(NodeNames);
+	for (FName NodeName : NodeNames)
+	{
+		if (Status != EStatus::Running)
+		{
+			// State graph was reset or paused during last node Start().
+			break;
+		}
+
+		FStateGraphNodeRef* Node = GetNodeRef(NodeName);
+		if (!Node)
+		{
+			++Removed;
+			continue;
+		}
+
+		switch ((*Node)->Status)
+		{
+		case FStateGraphNode::EStatus::NotStarted:
+		case FStateGraphNode::EStatus::Blocked:
+			if (!(*Node)->CheckDependencies())
+			{
+				if (!StateGraphWeakPtr.IsValid())
+				{
+					return;
+				}
+
+				// Get node again in case CheckDependencies() removed the node.
+				Node = GetNodeRef(NodeName);
+				if (Node)
+				{
+					(*Node)->SetStatus(FStateGraphNode::EStatus::Blocked);
+
+					if (!StateGraphWeakPtr.IsValid())
+					{
+						return;
+					}
+
+					++Blocked;
+				}
+				else
+				{
+					++Removed;
+				}
+
+				break;
+			}
+
+			UE_LOG_STATEGRAPH(Log, TEXT("[%s.%s] Starting node"), *Name.ToString(), *NodeName.ToString());
+			(*Node)->SetStatus(FStateGraphNode::EStatus::Started);
+
+			if (!StateGraphWeakPtr.IsValid())
+			{
+				return;
+			}
+
+			Node = GetNodeRef(NodeName);
+			if (!Node)
+			{
+				++Removed;
+				break;
+			}
+
+			(*Node)->StartTime = Now;
+			(*Node)->Start();
+
+			if (!StateGraphWeakPtr.IsValid())
+			{
+				return;
+			}
+
+			// Get node again in case Start() removed the node.
+			Node = GetNodeRef(NodeName);
+			if (Node)
+			{
+				if ((*Node)->Status == FStateGraphNode::EStatus::Completed)
+				{
+					// Start() called Complete() before returning.
+					++Completed;
+				}
+				else
+				{
+					++Started;
+				}
+			}
+			else
+			{
+				++Removed;
+			}
+
+			if (!Node || (*Node)->Status != FStateGraphNode::EStatus::Started)
+			{
+				break;
+			}
+
+			// Fall through to checkout for timeout.
+
+		case FStateGraphNode::EStatus::Started:
+			if ((*Node)->Timeout > 0.f)
+			{
+				const double NodeTimeout = ((*Node)->StartTime + (*Node)->Timeout) - Now;
+				if (NodeTimeout <= 0.f)
+				{
+					++TimedOut;
+					UE_LOG_STATEGRAPH(Log, TEXT("[%s.%s] Node timed out (Duration=%.6f)"), *Name.ToString(), *NodeName.ToString(), Now - (*Node)->StartTime);
+					(*Node)->SetStatus(FStateGraphNode::EStatus::TimedOut);
+
+					if (!StateGraphWeakPtr.IsValid())
+					{
+						return;
+					}
+
+					Node = GetNodeRef(NodeName);
+					if (Node)
+					{
+						(*Node)->TimedOut();
+						if (!StateGraphWeakPtr.IsValid())
+						{
+							return;
+						}
+					}
+
+					break;
+				}
+				
+				if (NextTimeout == 0.f || NodeTimeout < NextTimeout)
+				{
+					NextTimeout = NodeTimeout;
+				}
+			}
+
+			++Running;
+			break;
+
+		case FStateGraphNode::EStatus::Completed:
+			++Completed;
+			break;
+
+		case FStateGraphNode::EStatus::TimedOut:
+			++TimedOut;
+			break;
+
+		default:
+			UE_LOG_STATEGRAPH(Error, TEXT("[%s.%s] Unknown state"), *Name.ToString(), *NodeName.ToString());
+			checkNoEntry();
+			break;
+		}
+	}
+
+	bRunning = false;
+	UE_LOG_STATEGRAPH(Verbose, TEXT("[%s] Duration=%.6f Blocked=%d Started=%d Running=%d Completed=%d Removed=%d TimedOut=%d"),
+		*Name.ToString(), Now - StartTime, Blocked, Started, Running, Completed, Removed, TimedOut);
+
+	if (bRunAgain)
+	{
+		bRunAgain = false;
+		LogDebugInfo();
+		Run();
+		return;
+	}
+
+	if (Status != EStatus::Running)
+	{
+		// State graph was reset or paused during loop, don't change Status.
+	}
+	else if (Started == 0 && Running == 0)
+	{
+		if (Blocked == 0 && TimedOut == 0)
+		{
+			CompletedTime = Now;
+			UE_LOG_STATEGRAPH(Verbose, TEXT("[%s] Completed (Duration=%.6f)"), *Name.ToString(), CompletedTime - StartTime);
+			SetStatus(EStatus::Completed);
+		}
+		else
+		{
+			UE_LOG_STATEGRAPH(Warning, TEXT("[%s] Blocked on %d nodes, timed out %d nodes"), *Name.ToString(), Blocked, TimedOut);
+			LogDebugInfo(true);
+			SetStatus(EStatus::Blocked);
+		}
+	}
+	else
+	{
+		UE_LOG_STATEGRAPH(VeryVerbose, TEXT("[%s] Waiting on %d nodes"), *Name.ToString(), Started + Running);
+		SetStatus(EStatus::Waiting);
+	}
+
+	if (!StateGraphWeakPtr.IsValid())
+	{
+		return;
+	}
+
+	if ((Status == EStatus::Blocked || Status == EStatus::Waiting) && NextTimeout > 0.f)
+	{
+		TimeoutTicker = FTSTicker::GetCoreTicker().AddTicker(FTickerDelegate::CreateSPLambda(this, [this](float DeltaTime) {
+			Run();
+			return false;
+		}), NextTimeout);
+	}
+
+	LogDebugInfo();
+}
+
+void FStateGraph::Reset()
+{
+	// Keep a reference to check if the state graph is destroyed during external functions.
+	FStateGraphWeakPtr StateGraphWeakPtr(AsWeak());
+
+	UE_LOG_STATEGRAPH(Verbose, TEXT("[%s] Resetting state graph"), *Name.ToString());
+	StartTime = 0.f;
+	CompletedTime = 0.f;
+	bRunAgain = false;
+	SetStatus(EStatus::NotStarted);
+
+	if (!StateGraphWeakPtr.IsValid())
+	{
+		return;
+	}
+
+	// Copy names since node map may be modified during loop.
+	TArray<FName> NodeNames;
+	Nodes.GenerateKeyArray(NodeNames);
+	for (FName NodeName : NodeNames)
+	{
+		if (FStateGraphNodeRef* Node = GetNodeRef(NodeName))
+		{
+			(*Node)->Reset();
+			if (!StateGraphWeakPtr.IsValid())
+			{
+				return;
+			}
+		}
+	}
+}
+
+void FStateGraph::Pause()
+{
+	UE_LOG_STATEGRAPH(Verbose, TEXT("[%s] Pausing state graph"), *Name.ToString());
+	bRunAgain = false;
+	SetStatus(EStatus::Paused);
+}
+
+void FStateGraph::LogDebugInfo(bool bWarning)
+{
+	if (!bWarning && !UE_LOG_ACTIVE(LogStateGraph, VeryVerbose))
+	{
+		return;
+	}
+
+	FString Message = FString::Printf(TEXT("[%s] Status=%s Nodes=%d Duration=%.6f"), *Name.ToString(), GetStatusName(), Nodes.Num(), GetDuration());
+
+	if (bWarning)
+	{
+		UE_LOG_STATEGRAPH(Warning, TEXT("%s"), *Message);
+	}
+	else
+	{
+		UE_LOG_STATEGRAPH(VeryVerbose, TEXT("%s"), *Message);
+	}
+
+	for (auto Node : Nodes)
+	{
+		TArray<FString> Missing;
+		TArray<FString> Blocked;
+		TArray<FString> Completed;
+		TArray<FString> TimedOut;
+
+		for (FName Dependency : Node.Value->Dependencies)
+		{
+			FStateGraphNodeRef* DependencyNode = GetNodeRef(Dependency);
+			if (DependencyNode)
+			{
+				if ((*DependencyNode)->Status == FStateGraphNode::EStatus::Completed)
+				{
+					Completed.Add(Dependency.ToString());
+				}
+				else if ((*DependencyNode)->Status == FStateGraphNode::EStatus::Blocked)
+				{
+					Blocked.Add(Dependency.ToString());
+				}
+				else
+				{
+					TimedOut.Add(Dependency.ToString());
+				}
+			}
+			else
+			{
+				Missing.Add(Dependency.ToString());
+			}
+		}
+
+		TArray<FString> Dependencies;
+		if (Missing.Num())
+		{
+			Dependencies.Add(TEXT("Missing=") + FString::Join(Missing, TEXT(",")));
+		}
+		if (Blocked.Num())
+		{
+			Dependencies.Add(TEXT("Blocked=") + FString::Join(Blocked, TEXT(",")));
+		}
+		if (Completed.Num())
+		{
+			Dependencies.Add(TEXT("Completed=") + FString::Join(Completed, TEXT(",")));
+		}
+		if (TimedOut.Num())
+		{
+			Dependencies.Add(TEXT("TimedOut=") + FString::Join(TimedOut, TEXT(",")));
+		}
+		if (Dependencies.Num() == 0)
+		{
+			Dependencies.Add(TEXT("None"));
+		}
+
+		Message = FString::Printf(TEXT("[%s.%s] Status=%s Duration=%.6f Dependencies(%s)"),
+			*Name.ToString(), *Node.Key.ToString(), Node.Value->GetStatusName(), Node.Value->GetDuration(), *FString::Join(Dependencies, TEXT(" ")));
+
+		if (bWarning)
+		{
+			UE_LOG_STATEGRAPH(Warning, TEXT("%s"), *Message);
+		}
+		else
+		{
+			UE_LOG_STATEGRAPH(VeryVerbose, TEXT("%s"), *Message);
+		}
+	}
+}
+
+void FStateGraph::SetStatus(EStatus NewStatus)
+{
+	if (Status != NewStatus)
+	{
+		const EStatus OldStatus = Status;
+		Status = NewStatus;
+		OnStatusChanged.Broadcast(*this, OldStatus, NewStatus);
+	}
+}
+
+void FStateGraph::UpdateConfig()
+{
+	GConfig->GetDouble(*ConfigSectionName, TEXT("Timeout"), Timeout, GEngineIni);
+}
+
+void FStateGraph::OnConfigSectionsChanged(const FString& IniFilename, const TSet<FString>& SectionNames)
+{
+	if (IniFilename != GEngineIni)
+	{
+		return;
+	}
+
+	// Keep a reference to check if the state graph is destroyed during external functions.
+	FStateGraphWeakPtr StateGraphWeakPtr(AsWeak());
+
+	for (const FString& SectionName : SectionNames)
+	{
+		// Assume all node section names start with the state graph section name.
+		if (!SectionName.StartsWith(ConfigSectionName))
+		{
+			continue;
+		}
+
+		if (SectionName.Len() == ConfigSectionName.Len())
+		{
+			UpdateConfig();
+			continue;
+		}
+
+		// Copy names since node map may be modified during loop.
+		TArray<FName> NodeNames;
+		Nodes.GenerateKeyArray(NodeNames);
+		for (FName NodeName : NodeNames)
+		{
+			FStateGraphNodeRef* Node = GetNodeRef(NodeName);
+			if (Node)
+			{
+				if (SectionName == (*Node)->ConfigSectionName)
+				{
+					(*Node)->UpdateConfig();
+					if (!StateGraphWeakPtr.IsValid())
+					{
+						return;
+					}
+				}
+			}
+		}
+	}
+}
+
+} // UE
diff --git a/Engine/Plugins/Experimental/StateGraph/Source/StateGraph/Public/StateGraph.h b/Engine/Plugins/Experimental/StateGraph/Source/StateGraph/Public/StateGraph.h
new file mode 100644
--- /dev/null
+++ b/Engine/Plugins/Experimental/StateGraph/Source/StateGraph/Public/StateGraph.h
@@ -0,0 +1,427 @@
+// Copyright Epic Games, Inc. All Rights Reserved.
+
+#pragma once
+
+/**
+ * This module provides a generic state machine framework. For an example, see StateGraphTestsExample.cpp. Features include:
+ * - Free-form structure with no required linear or hierarchical structure - just graph nodes with dependencies.
+ * - Support for async tasks via a delayed function call to complete nodes.
+ * - Modifiable at runtime to add, remove, rename, or reset nodes in the graph.
+ * - Able to modify a state graph easily in derived classes or nested objects.
+ * - Each node is a full object that can encapsulate all required functions and data.
+ * - Wrapper node for delegates to support various function types directly.
+ * - Max execution depth of one node at a time.
+ * - Timeout support for state graph and individual nodes.
+ * - Configurable and hotfixable via ini settings.
+ */
+
+#include "Containers/Ticker.h"
+#include "Delegates/Delegate.h"
+#include "Logging/LogMacros.h"
+#include "Misc/CoreMisc.h"
+#include "Templates/Function.h"
+#include "Templates/IsClass.h"
+#include "Templates/SharedPointer.h"
+#include "Templates/UnrealTemplate.h"
+#include "UObject/NameTypes.h"
+
+STATEGRAPH_API DECLARE_LOG_CATEGORY_EXTERN(LogStateGraph, Log, All);
+#define UE_LOG_STATEGRAPH(Verbosity, Format, ...) UE_LOG(LogStateGraph, Verbosity, TEXT("[%s] ") Format, ANSI_TO_TCHAR(__FUNCTION__), ##__VA_ARGS__)
+
+namespace UE {
+
+class FStateGraph;
+using FStateGraphRef = TSharedRef<FStateGraph, ESPMode::ThreadSafe>;
+using FStateGraphPtr = TSharedPtr<FStateGraph, ESPMode::ThreadSafe>;
+using FStateGraphWeakPtr = TWeakPtr<FStateGraph, ESPMode::ThreadSafe>;
+
+class FStateGraphNode;
+using FStateGraphNodeRef = TSharedRef<FStateGraphNode, ESPMode::ThreadSafe>;
+using FStateGraphNodePtr = TSharedPtr<FStateGraphNode, ESPMode::ThreadSafe>;
+using FStateGraphNodeWeakPtr = TWeakPtr<FStateGraphNode, ESPMode::ThreadSafe>;
+
+class FStateGraphNodeFunction;
+using FStateGraphNodeFunctionRef = TSharedRef<FStateGraphNodeFunction, ESPMode::ThreadSafe>;
+using FStateGraphNodeFunctionPtr = TSharedPtr<FStateGraphNodeFunction, ESPMode::ThreadSafe>;
+using FStateGraphNodeFunctionWeakPtr = TWeakPtr<FStateGraphNodeFunction, ESPMode::ThreadSafe>;
+
+/**
+ * Abstract base class for all nodes in the state graph.
+ */
+class STATEGRAPH_API FStateGraphNode : public TSharedFromThis<FStateGraphNode>, public FNoncopyable
+{
+public:
+
+	FStateGraphNode(FName InName);
+
+	virtual ~FStateGraphNode();
+
+	/** Get the name of the node. */
+	FName GetName() const
+	{
+		return Name;
+	}
+
+	enum class EStatus : uint8
+	{
+		/** State graph has not attempted to start this node yet, or the node has been reset. */
+		NotStarted,
+		/** State graph attempted to start this node, but it is waiting on dependencies or some external blocking condition. */
+		Blocked,
+		/** State graph has started this node by calling Start(). */
+		Started,
+		/** Node's Complete() function has been called, allowing dependent nodes to run. */
+		Completed,
+		/** Node timed out. */
+		TimedOut
+	};
+
+	/** Get the status of the node. */
+	EStatus GetStatus() const
+	{
+		return Status;
+	}
+
+	/** Get the status of the node as a string. */
+	const TCHAR* GetStatusName() const
+	{
+		return GetStatusName(Status);
+	}
+
+	static const TCHAR* GetStatusName(EStatus Status);
+
+	/** Get timeout in seconds. */
+	double GetTimeout() const
+	{
+		return Timeout;
+	}
+
+	/** Set timeout in seconds from when Start is run until Complete is run. If this timeout is exceeded, the Status is set to TimedOut. */
+	void SetTimeout(double InTimeout);
+
+	/** Get duration of time that the node has taken to run. Returns 0 if not yet started. */
+	double GetDuration() const;
+
+	/** Returns true if all dependencies are met and the node can be started. This base version checks that all nodes in the Dependencies set have been completed, but can be replaced or extended. */
+	virtual bool CheckDependencies() const;
+
+	/** Set Status to Completed and notify the state graph to start other nodes which may have been blocked on this node. */
+	virtual void Complete();
+
+	/** Reset this node to NotStarted so it is run again. Override to reset any state or clean up pending async tasks. */
+	virtual void Reset();
+
+	/** Get the state graph for this node, which may be invalid if this node has not been added or was removed. */
+	FStateGraphPtr GetStateGraph()
+	{
+		return StateGraphWeakPtr.Pin();
+	}
+
+	/** Get the name of the state graph, or NAME_None if the state graph is invalid. */
+	FName GetStateGraphName() const;
+
+	/**
+	 * Helper function to create a new node, add it to this node's StateGraph, and add this node as a dependency.
+	 * This allows for chaining nodes together such as:
+	 *
+	 *	StateGraph->CreateNode("NodeA", this, &SomeMemberFunction)
+	 *		->Next("NodeB", &SomeStaticFunction)
+	 *		->Next<SomeNodeClass>()
+	 *		->Next("NodeC", [](FStateGraphNodeFunctionComplete Complete) { Complete(); return true; });
+	 */
+	template<typename NodeType, typename... ArgsTypes>
+	TSharedPtr<NodeType, ESPMode::ThreadSafe> Next(ArgsTypes&&... Args) const;
+
+	template<typename... ArgsTypes>
+	FStateGraphNodeFunctionPtr Next(ArgsTypes&&... Args) const;
+
+	TSet<FName> Dependencies;
+
+protected:
+
+	/** Start this node. */
+	virtual void Start() = 0;
+
+	/** Node timed out. Override to cleanup up any state or pending async tasks. This will not get called if OnNodeStatusChanged removes the node after updating the status to TimedOut. */
+	virtual void TimedOut() {}
+
+	/** Node was removed from the state graph. Override to cleanup up any state or pending async tasks. */
+	virtual void Removed() {}
+
+	/** Update config from INI files. Called when node is added to a state graph or during hotfixes. Override this to read any custom config the for the node. */
+	virtual void UpdateConfig();
+
+	/** Get the name of the node. */
+	const FString& GetConfigSectionName() const
+	{
+		return ConfigSectionName;
+	}
+
+private:
+
+	/** Set status and fire delegate if needed. */
+	void SetStatus(EStatus NewStatus);
+
+	FName Name;
+	FString ConfigSectionName;
+	EStatus Status = EStatus::NotStarted;
+	FStateGraphWeakPtr StateGraphWeakPtr;
+	double StartTime = 0.f;
+	double CompletedTime = 0.f;
+	double Timeout = 0.f;
+
+	friend FStateGraph;
+};
+
+/** Complete function type passed into FStateGraphNodeFunction node start functions. */
+using FStateGraphNodeFunctionComplete = TFunction<void()>;
+
+/** Start function delegate type for FStateGraphNodeFunction nodes. */
+DECLARE_DELEGATE_OneParam(FStateGraphNodeFunctionStart, FStateGraphNodeFunctionComplete /*Complete*/);
+
+/**
+ * Node for wrapping functions as delegates.
+ */
+class STATEGRAPH_API FStateGraphNodeFunction : public FStateGraphNode
+{
+public:
+
+	FStateGraphNodeFunction(FName InName, const FStateGraphNodeFunctionStart& InStartFunction);
+	virtual bool CheckDependencies() const override;
+
+protected:
+
+	virtual void Start() override;
+	FStateGraphNodeFunctionStart StartFunction;
+
+	friend FStateGraph;
+};
+
+/**
+ * State graph class to run and manage nodes with.
+ */
+class STATEGRAPH_API FStateGraph : public TSharedFromThis<FStateGraph>, public FNoncopyable
+{
+public:
+
+	FStateGraph(FName InName);
+
+	virtual ~FStateGraph();
+
+	/** Perform state graph initialization that could not run in constructor due to shared pointers not being ready, such as registering for hotfixes. */
+	void Initialize();
+
+	/** Get the name of the state graph. */
+	FName GetName() const
+	{
+		return Name;
+	}
+
+	enum class EStatus : uint8
+	{
+		/** State graph has not started yet. */
+		NotStarted,
+		/** State graph is running. */
+		Running,
+		/** State graph is waiting for one or more nodes to complete. */
+		Waiting,
+		/** State graph is blocked due to missing dependencies. */
+		Blocked,
+		/** State graph has completed. */
+		Completed,
+		/** State graph has been paused. */
+		Paused,
+		/** State graph timed out and stopped running. */
+		TimedOut
+	};
+
+	/** Get the status of the state graph. */
+	EStatus GetStatus() const
+	{
+		return Status;
+	}
+
+	/** Get the status of the state graph as a string. */
+	const TCHAR* GetStatusName() const
+	{
+		return GetStatusName(Status);
+	}
+
+	static const TCHAR* GetStatusName(EStatus Status);
+
+	/** Delegate fired when the state graph status changes. Useful to know when it becomes blocked or times out. */
+	DECLARE_MULTICAST_DELEGATE_ThreeParams(FStatusChanged, FStateGraph& /*StateGraph*/, EStatus /*OldStatus*/, EStatus /*NewStatus*/ );
+	FStatusChanged OnStatusChanged;
+
+	/** Delegate fired when a node status changes. */
+	DECLARE_MULTICAST_DELEGATE_ThreeParams(FNodeStatusChanged, FStateGraphNode& /*Node*/, FStateGraphNode::EStatus /*OldStatus*/, FStateGraphNode::EStatus /*NewStatus*/ );
+	FNodeStatusChanged OnNodeStatusChanged;
+
+	/** Get timeout in seconds. */
+	double GetTimeout() const
+	{
+		return Timeout;
+	}
+
+	/** Set timeout in seconds from the initial run call (when Status changes from NotStarted to Running) to when Status changes to Completed. */
+	void SetTimeout(double InTimeout);
+
+	/** Get duration of time that the state graph has taken to run. Returns 0 if not yet started. */
+	double GetDuration() const;
+
+	/** Add a node to the state graph. Returns true if successful, false if the node is either already added to a state graph or if an existing node with the same name exists. */
+	bool AddNode(const FStateGraphNodeRef& Node);
+
+	/** Remove a node from the state graph. Returns true if successful, false if no node was found for the name. */
+	bool RemoveNode(FName NodeName);
+
+	/** Remove all nodes from the state graph that exist when this function is started. It's possible for nodes to be added back before this function returns if Node->Removed() adds nodes. */
+	void RemoveAllNodes();
+
+	/** Get a node from the state graph. References should not be cached in case the node is destroyed. */
+	FStateGraphNodeRef* GetNodeRef(FName NodeName)
+	{
+		return Nodes.Find(NodeName);
+	}
+
+	/** Get a node from the state graph, casting to a derived type if needed. */
+	template<typename NodeType = FStateGraphNode>
+	TSharedPtr<NodeType, ESPMode::ThreadSafe> GetNode(FName NodeName)
+	{
+		TSharedPtr<NodeType, ESPMode::ThreadSafe> NodePtr;
+		if (FStateGraphNodeRef* NodeRef = GetNodeRef(NodeName))
+		{
+			NodePtr = StaticCastSharedRef<NodeType>(*NodeRef);
+		}
+
+		return NodePtr;
+	}
+
+	/** Create and add a new node for a given node class type. */
+	template<typename NodeType, typename... ArgsTypes>
+	TSharedRef<NodeType, ESPMode::ThreadSafe> CreateNode(ArgsTypes&&... Args)
+	{
+		TSharedRef<NodeType, ESPMode::ThreadSafe> Node(MakeShared<NodeType>(Forward<ArgsTypes>(Args)...));
+		ensureAlways(AddNode(Node));
+		return Node;
+	}
+
+	/** Create a new function node for member functions. */
+	template<typename ObjectType, typename FunctionType, typename... ArgsTypes>
+	typename TEnableIf<TIsClass<ObjectType>::Value && TIsInvocable<FunctionType, ObjectType, FStateGraphNodeFunctionComplete, ArgsTypes...>::Value, FStateGraphNodeFunctionRef>::Type
+		CreateNode(FName NodeName, ObjectType* Object, FunctionType Function, ArgsTypes&&... Args)
+	{
+		if constexpr (IsDerivedFromSharedFromThis<ObjectType>())
+		{
+			return CreateNode<FStateGraphNodeFunction>(NodeName, FStateGraphNodeFunctionStart::CreateSP(Object, Function, Forward<ArgsTypes>(Args)...));
+		}
+		else if constexpr (TIsDerivedFrom<ObjectType, UObject>::IsDerived)
+		{
+			return CreateNode<FStateGraphNodeFunction>(NodeName, FStateGraphNodeFunctionStart::CreateUObject(Object, Function, Forward<ArgsTypes>(Args)...));
+		}
+		else if constexpr (TIsMemberPointer<FunctionType>::Value)
+		{
+			return CreateNode<FStateGraphNodeFunction>(NodeName, FStateGraphNodeFunctionStart::CreateRaw(Object, Function, Forward<ArgsTypes>(Args)...));
+		}
+	}
+
+	/** Create a new function node for lambdas tied to objects. */
+	template<typename ObjectType, typename FunctionType, typename... ArgsTypes>
+	typename TEnableIf<TIsClass<ObjectType>::Value && TIsInvocable<FunctionType, FStateGraphNodeFunctionComplete, ArgsTypes...>::Value, FStateGraphNodeFunctionRef>::Type
+		CreateNode(FName NodeName, ObjectType* Object, FunctionType Function, ArgsTypes&&... Args)
+	{
+		if constexpr (IsDerivedFromSharedFromThis<ObjectType>())
+		{
+			return CreateNode<FStateGraphNodeFunction>(NodeName, FStateGraphNodeFunctionStart::CreateSPLambda(Object, Function, Forward<ArgsTypes>(Args)...));
+		}
+		else if constexpr (TIsDerivedFrom<ObjectType, UObject>::IsDerived)
+		{
+			return CreateNode<FStateGraphNodeFunction>(NodeName, FStateGraphNodeFunctionStart::CreateWeakLambda(Object, Function, Forward<ArgsTypes>(Args)...));
+		}
+	}
+
+	/** Create a new function node for lambas and static methods. */
+	template<typename FunctionType, typename... ArgsTypes>
+	typename TEnableIf<TIsInvocable<FunctionType, FStateGraphNodeFunctionComplete, ArgsTypes...>::Value, FStateGraphNodeFunctionRef>::Type
+		CreateNode(FName NodeName, FunctionType Function, ArgsTypes&&... Args)
+	{
+		return CreateNode<FStateGraphNodeFunction>(NodeName, FStateGraphNodeFunctionStart::CreateLambda(Function, Forward<ArgsTypes>(Args)...));
+	}
+
+	/**
+	 * Run the state graph. This is run automatically when nodes complete to try starting any blocked nodes, so this only needs to be run manually when:
+	 * - The state graph has not been run yet (Status = NotStarted).
+	 * - The state graph has become blocked, and the blocking condition has been resolved (Status = Blocked).
+	 * - A new node has been added to the state graph or dependencies for nodes have changed that may unblock existing nodes.
+	 * - The state graph was reset or paused and should be started again.
+	 */
+	void Run();
+
+	/**
+	 * Stop and reset the state graph and all nodes. The state graph can be started again with Run().
+	 */
+	void Reset();
+
+	/**
+	 * Pause the state graph and prevent additional nodes from being started. Other actions can still happen such as nodes completing nodes, adding nodes, resetting, etc.
+	 */
+	void Pause();
+
+	/** Log debug info for the state graph and each node. */
+	void LogDebugInfo(bool bWarning = false);
+
+protected:
+
+	/** Set status and fire delegate if needed. */
+	void SetStatus(EStatus NewStatus);
+
+	/** Update config from INI files. Called from Initialize and during hotfixes. */
+	void UpdateConfig();
+
+	/** Callback to update state graph and node configuration during hotfixes. */
+	void OnConfigSectionsChanged(const FString& IniFilename, const TSet<FString>& SectionNames);
+
+	FName Name;
+	FString ConfigSectionName;
+	EStatus Status = EStatus::NotStarted;
+	double StartTime = 0.f;
+	double CompletedTime = 0.f;
+	double Timeout = 0.f;
+	TMap<FName, FStateGraphNodeRef> Nodes;
+	bool bRunning = false;
+	bool bRunAgain = false;
+	FTSTicker::FDelegateHandle TimeoutTicker;
+	FDelegateHandle ConfigSectionsChangedDelegate;
+};
+
+// FStateGraphNode templates that need FStateGraph defined.
+
+template<typename NodeType, typename... ArgsTypes>
+TSharedPtr<NodeType, ESPMode::ThreadSafe> FStateGraphNode::Next(ArgsTypes&&... Args) const
+{
+	if (!StateGraphWeakPtr.IsValid())
+	{
+		UE_LOG_STATEGRAPH(Warning, TEXT("[%s] Node attempted to be created with invalid state graph"), *Name.ToString());
+		return TSharedPtr<NodeType, ESPMode::ThreadSafe>();
+	}
+
+	TSharedPtr<NodeType, ESPMode::ThreadSafe> Node(StateGraphWeakPtr.Pin()->CreateNode<NodeType>(Forward<ArgsTypes>(Args)...));
+	Node->Dependencies.Add(Name);
+	return Node;
+}
+
+template<typename... ArgsTypes>
+FStateGraphNodeFunctionPtr FStateGraphNode::Next(ArgsTypes&&... Args) const
+{
+	if (!StateGraphWeakPtr.IsValid())
+	{
+		UE_LOG_STATEGRAPH(Warning, TEXT("[%s] Node attempted to be created with invalid state graph"), *Name.ToString());
+		return FStateGraphNodeFunctionPtr();
+	}
+
+	FStateGraphNodeFunctionPtr Node(StateGraphWeakPtr.Pin()->CreateNode(Forward<ArgsTypes>(Args)...));
+	Node->Dependencies.Add(Name);
+	return Node;
+}
+
+} // UE
diff --git a/Engine/Plugins/Experimental/StateGraph/Source/StateGraph/StateGraph.Build.cs b/Engine/Plugins/Experimental/StateGraph/Source/StateGraph/StateGraph.Build.cs
new file mode 100644
--- /dev/null
+++ b/Engine/Plugins/Experimental/StateGraph/Source/StateGraph/StateGraph.Build.cs
@@ -0,0 +1,13 @@
+// Copyright Epic Games, Inc. All Rights Reserved.
+
+using UnrealBuildTool;
+
+public class StateGraph : ModuleRules
+{
+	public StateGraph(ReadOnlyTargetRules Target) : base(Target)
+	{
+		PrivateDependencyModuleNames.AddRange(new string[] {
+			"Core"
+		});
+	}
+}
diff --git a/Engine/Plugins/Experimental/StateGraph/StateGraph.uplugin b/Engine/Plugins/Experimental/StateGraph/StateGraph.uplugin
new file mode 100644
--- /dev/null
+++ b/Engine/Plugins/Experimental/StateGraph/StateGraph.uplugin
@@ -0,0 +1,18 @@
+{
+	"FileVersion": 3,
+	"Version": 1,
+	"VersionName": "1.0",
+	"FriendlyName": "State Graph",
+	"Description": "Generic state machine management class.",
+	"Category": "Other",
+	"CreatedBy": "Epic Games, Inc.",
+	"CreatedByURL": "https://epicgames.com",
+	"EnabledByDefault": false,
+	"Modules": [
+		{
+			"Name": "StateGraph",
+			"Type": "Runtime",
+			"LoadingPhase": "PostConfigInit"
+		}
+	]
+}
diff --git a/Engine/Source/Programs/StateGraphTests/Source/StateGraphTestsExample.cpp b/Engine/Source/Programs/StateGraphTests/Source/StateGraphTestsExample.cpp
new file mode 100644
--- /dev/null
+++ b/Engine/Source/Programs/StateGraphTests/Source/StateGraphTestsExample.cpp
@@ -0,0 +1,237 @@
+// Copyright Epic Games, Inc. All Rights Reserved.
+
+#include "Containers/Ticker.h"
+#include "Misc/ConfigCacheIni.h"
+#include "StateGraph.h"
+#include "TestHarness.h"
+
+namespace UE::StateGraphTests::Example
+{
+
+// This example shows many of the features provided by the StateGraph module. See class descriptions below for how nodes relate to each other within the state graph.
+// Running this produces the following debug output at the end:
+//
+// [UE::FStateGraph::Run] [Example] Completed (Duration=0.001272)
+// [UE::FStateGraph::LogDebugInfo] [Example] Status=Completed Nodes=11 Duration=0.001272
+// [UE::FStateGraph::LogDebugInfo] [Example.Base.NodeA] Status=Completed Duration=0.000023 Dependencies(None)
+// [UE::FStateGraph::LogDebugInfo] [Example.Base.NodeB] Status=Completed Duration=0.000258 Dependencies(None)
+// [UE::FStateGraph::LogDebugInfo] [Example.Base.NodeC] Status=Completed Duration=0.000206 Dependencies(Completed=Base.NodeA,Base.NodeB)
+// [UE::FStateGraph::LogDebugInfo] [Example.Base.NodeD] Status=Completed Duration=0.000013 Dependencies(Completed=Base.NodeC,Derived.NodeF)
+// [UE::FStateGraph::LogDebugInfo] [Example.Base.NodeE] Status=Completed Duration=0.000190 Dependencies(Completed=Base.NodeC)
+// [UE::FStateGraph::LogDebugInfo] [Example.Derived.NodeF] Status=Completed Duration=0.000556 Dependencies(Completed=Base.NodeC)
+// [UE::FStateGraph::LogDebugInfo] [Example.Derived.NodeG] Status=Completed Duration=0.000056 Dependencies(None)
+// [UE::FStateGraph::LogDebugInfo] [Example.Derived.NodeH] Status=Completed Duration=0.000370 Dependencies(Completed=Derived.NodeG)
+// [UE::FStateGraph::LogDebugInfo] [Example.Module.NodeA] Status=Completed Duration=0.000014 Dependencies(None)
+// [UE::FStateGraph::LogDebugInfo] [Example.Module.NodeB] Status=Completed Duration=0.000321 Dependencies(Completed=Module.NodeA)
+// [UE::FStateGraph::LogDebugInfo] [Example.Module.NodeC] Status=Completed Duration=0.000015 Dependencies(Completed=Module.NodeB)
+//
+// To see the graph work through all nodes, build and run: StateGraphTests.exe [Example] --log --extra-args -LogCmds="LogStateGraph All" -AbsLog=D:\StateGraph.log
+
+
+// Use custom class nodes like this to manage functions and data for a single node, as opposed to the single function nodes seen below.
+// Instances of this are added to the state graph in the classes below.
+class FExampleClassNode : public FStateGraphNode
+{
+public:
+
+	FExampleClassNode(FName InName) :
+		FStateGraphNode(InName)
+	{}
+
+	float Input = 0.1f;
+	float Output = 0;
+
+protected:
+
+	virtual void Start()
+	{
+		// Code to setup async task...
+		FTSTicker::GetCoreTicker().AddTicker(FTickerDelegate::CreateSP(this, &FExampleClassNode::RunAsync), Input);
+	}
+
+	bool RunAsync(float DeltaTime)
+	{
+		if (++RunState < CustomConfig)
+		{
+			return true;
+		}
+
+		FinishAsync();
+		return false;
+	}
+
+	void FinishAsync()
+	{
+		// Code to process result of async task...
+		Output = 0.2f;
+		Complete();
+	}
+
+	virtual void UpdateConfig()
+	{
+		// Override to read custom config for the node.
+		FStateGraphNode::UpdateConfig();
+		GConfig->GetInt(*GetConfigSectionName(), TEXT("CustomConfig"), CustomConfig, GEngineIni);
+	}
+
+	int32 RunState = 0;
+	int32 CustomConfig = 3;
+};
+
+using FExampleClassNodePtr = TSharedPtr<FExampleClassNode, ESPMode::ThreadSafe>;
+
+
+// Helper function to simulate async tasks in simple function nodes below.
+void ExampleAsync(float Delay, const FStateGraphNodeFunctionComplete& Complete)
+{
+	FTSTicker::GetCoreTicker().AddTicker(FTickerDelegate::CreateLambda([Complete](float DeltaTime) {
+		Complete();
+		return false;
+	}), Delay);
+}
+
+
+// FExampleBase - Base class of some function nodes that perform some work.
+class FExampleBase : public TSharedFromThis<FExampleBase>
+{
+public:
+
+	FExampleBase(const FStateGraphRef& InStateGraph) :
+		StateGraph(InStateGraph)
+	{}
+
+	virtual ~FExampleBase() = default;
+
+	virtual void SetupStateGraph()
+	{
+		// CreateNode supports any function type that can be used to create delegates with: lambdas, members, static, raw, etc.
+
+		// Since NodeA and NodeB have no dependencies, they run first concurrently.
+		StateGraph->CreateNode("Base.NodeA", this, [](FStateGraphNodeFunctionComplete Complete) {	Complete(); });
+		StateGraph->CreateNode("Base.NodeB", this, [](FStateGraphNodeFunctionComplete Complete) {	ExampleAsync(0.1f, Complete); });
+
+		// CreateNode also supports custom node types by explicitly specifying the type and passing arguments to the constructor.
+		FExampleClassNodePtr NodeC = StateGraph->CreateNode<FExampleClassNode>("Base.NodeC");
+		// Add dependencies so NodeC isn't run until both NodeA and NodeB complete.
+		NodeC->Dependencies.Append({ "Base.NodeA", "Base.NodeB" });
+		NodeC->Input = 0.1f;
+
+		// Both NodeD and NodeE will run concurrently once NodeC is complete.
+		StateGraph->CreateNode("Base.NodeD", &FExampleBase::NodeD)
+			->Dependencies.Add("Base.NodeC");
+		StateGraph->CreateNode("Base.NodeE", this, &FExampleBase::NodeE)
+			->Dependencies.Add("Base.NodeC");
+	}
+
+protected:
+
+	static void NodeD(FStateGraphNodeFunctionComplete Complete)
+	{
+		Complete();
+	}
+
+	void NodeE(FStateGraphNodeFunctionComplete Complete)
+	{
+		FExampleClassNodePtr NodeC = StateGraph->GetNode<FExampleClassNode>("Base.NodeC");
+		check(NodeC.IsValid());
+		// Use NodeC Output within this node...
+		ExampleAsync(NodeC->Output, Complete);
+	}
+
+	FStateGraphRef StateGraph;
+};
+
+
+// FExampleModule - Class representing a module that shares no nodes with FExampleBase or FExampleDerived, but is used by FExampleDerived to complete some work.
+DECLARE_DELEGATE_OneParam(FExampleModuleComplete, int32 /* Output */);
+class FExampleModule : public TSharedFromThis<FExampleModule>
+{
+public:
+	FExampleModule(const FStateGraphRef& InStateGraph) :
+		StateGraph(InStateGraph)
+	{}
+
+	void Run(const FExampleModuleComplete& ModuleComplete)
+	{
+		// Run nodes in sequence to produce some output that the caller needs, using the same state graph which may be running unrelated nodes.
+		StateGraph->CreateNode("Module.NodeA", this, [](FStateGraphNodeFunctionComplete Complete) { Complete(); })
+			->Next("Module.NodeB", this, &FExampleModule::NodeB)
+			->Next("Module.NodeC", this, [ModuleComplete](FStateGraphNodeFunctionComplete Complete) { Complete(); ModuleComplete.ExecuteIfBound(42); });
+	}
+
+protected:
+
+	void NodeB(FStateGraphNodeFunctionComplete Complete)
+	{
+		ExampleAsync(0.1f, Complete);
+	}
+
+	FStateGraphRef StateGraph;
+};
+
+using FExampleModulePtr = TSharedPtr<FExampleModule, ESPMode::ThreadSafe>;
+
+
+// FExampleDerived - Derived class that updates the function nodes FExampleBase setup.
+class FExampleDerived : public FExampleBase
+{
+public:
+
+	FExampleDerived(const FStateGraphRef& InStateGraph) :
+		FExampleBase(InStateGraph)
+	{}
+
+	virtual void SetupStateGraph() override
+	{
+		FExampleBase::SetupStateGraph();
+
+		// Insert NodeF between NodeC and NodeD
+		StateGraph->CreateNode("Derived.NodeF", this, &FExampleDerived::NodeF)
+			->Dependencies.Add("Base.NodeC");
+
+		FStateGraphNodeRef* NodeDRef = StateGraph->GetNodeRef("Base.NodeD");
+		check(NodeDRef);
+		(*NodeDRef)->Dependencies.Add("Derived.NodeF");
+
+		// Add some additional nodes.
+		StateGraph->CreateNode("Derived.NodeG", this, [](FStateGraphNodeFunctionComplete Complete) { Complete(); })
+			->Next<FExampleClassNode>("Derived.NodeH");
+	}
+
+protected:
+
+	void NodeF(FStateGraphNodeFunctionComplete Complete)
+	{
+		FExampleModulePtr ModulePtr = MakeShared<FExampleModule>(StateGraph);
+		ModulePtr->Run(FExampleModuleComplete::CreateSPLambda(this, [ModulePtr, Complete](int32 Output)
+		{
+			// Use output...
+			Complete();
+		}));
+	}
+};
+
+
+TEST_CASE("Example demonstrating the features of StateGraph", "[Example]")
+{
+	// Engine does this, but it's needed here to ensure GConfig is initialized for the test.
+	FConfigCacheIni::InitializeConfigSystem();
+
+	FStateGraphRef StateGraph(MakeShared<FStateGraph>("Example"));
+	StateGraph->Initialize();
+
+	TSharedRef<FExampleDerived> Example = MakeShared<FExampleDerived>(StateGraph);
+	Example->SetupStateGraph();
+
+	StateGraph->Run();
+
+	// Tick the fake async tasks until complete.
+	while (StateGraph->GetStatus() == FStateGraph::EStatus::Waiting)
+	{
+		FTSTicker::GetCoreTicker().Tick(0.1f);
+	}
+
+	CHECK(StateGraph->GetStatus() == FStateGraph::EStatus::Completed);
+}
+
+} // UE::StateGraphTests
diff --git a/Engine/Source/Programs/StateGraphTests/Source/StateGraphTestsUnit.cpp b/Engine/Source/Programs/StateGraphTests/Source/StateGraphTestsUnit.cpp
new file mode 100644
--- /dev/null
+++ b/Engine/Source/Programs/StateGraphTests/Source/StateGraphTestsUnit.cpp
@@ -0,0 +1,467 @@
+// Copyright Epic Games, Inc. All Rights Reserved.
+
+#include "Misc/ConfigCacheIni.h"
+#include "Misc/CoreDelegates.h"
+#include "StateGraph.h"
+#include "TestHarness.h"
+
+namespace UE::StateGraphTests::Unit
+{
+
+TEST_CASE("FStateGraph Basic Tests", "[FStateGraph]")
+{
+	FStateGraphRef StateGraph(MakeShared<FStateGraph>("Test"));
+	CHECK(StateGraph->GetName() == "Test");
+	CHECK(StateGraph->GetStatus() == FStateGraph::EStatus::NotStarted);
+	CHECK(FString(StateGraph->GetStatusName()) == FString(TEXT("NotStarted")));
+	StateGraph->Run();
+	CHECK(StateGraph->GetStatus() == FStateGraph::EStatus::Completed);
+}
+
+class FTestNode : public FStateGraphNode
+{
+public:
+	FTestNode(FName InName) :
+		FStateGraphNode(InName)
+	{}
+
+	FTestNode(FName InName, const TSet<FName>& InDependencies) :
+		FStateGraphNode(InName)
+	{
+		Dependencies.Append(InDependencies);
+	}
+
+	virtual void Start() override {}
+	virtual void Removed() override { bWasRemovedCalled = true; }
+	virtual void TimedOut() override { bWasTimedOutCalled = true; }
+	virtual void UpdateConfig() override
+	{
+		FStateGraphNode::UpdateConfig();
+		GConfig->GetInt(*GetConfigSectionName(), TEXT("CustomConfig"), CustomConfig, GEngineIni);
+	}
+
+	bool bWasRemovedCalled = false;
+	bool bWasTimedOutCalled = false;
+	int32 CustomConfig = 1;
+};
+
+using FTestNodeRef = TSharedRef<FTestNode, ESPMode::ThreadSafe>;
+using FTestNodePtr = TSharedPtr<FTestNode, ESPMode::ThreadSafe>;
+using FTestNodeWeakPtr = TWeakPtr<FTestNode, ESPMode::ThreadSafe>;
+
+TEST_CASE("FStateGraphNode Basic Tests", "[FStateGraphNode]")
+{
+	FStateGraphRef StateGraph(MakeShared<FStateGraph>("Test"));
+	FStateGraphNodeRef TestNodeA = StateGraph->CreateNode<FTestNode>("TestNodeA");
+	CHECK(TestNodeA->GetName() == "TestNodeA");
+	CHECK(TestNodeA->GetStatus() == FStateGraphNode::EStatus::NotStarted);
+	CHECK(FString(TestNodeA->GetStatusName()) == FString(TEXT("NotStarted")));
+	CHECK(&TestNodeA.Get() == &StateGraph->GetNodeRef("TestNodeA")->Get());
+	CHECK(&TestNodeA.Get() == StateGraph->GetNode<FTestNode>("TestNodeA").Get());
+
+	StateGraph->Run();
+	CHECK(StateGraph->GetStatus() == FStateGraph::EStatus::Waiting);
+	CHECK(TestNodeA->GetStatus() == FStateGraphNode::EStatus::Started);
+
+	TestNodeA->Complete();
+	CHECK(StateGraph->GetStatus() == FStateGraph::EStatus::Completed);
+	CHECK(TestNodeA->GetStatus() == FStateGraphNode::EStatus::Completed);
+}
+
+TEST_CASE("FStateGraphNodeFunction Basic Tests", "[FStateGraphNodeFunction]")
+{
+	FStateGraphRef StateGraph(MakeShared<FStateGraph>("Test"));
+	FStateGraphNodeFunctionComplete Complete;
+	FStateGraphNodeFunctionRef TestNodeA = StateGraph->CreateNode("TestNodeA", [&Complete](FStateGraphNodeFunctionComplete InComplete) { Complete = InComplete; });
+	CHECK(TestNodeA->GetStatus() == FStateGraphNode::EStatus::NotStarted);
+	CHECK(FString(TestNodeA->GetStatusName()) == FString(TEXT("NotStarted")));
+
+	StateGraph->Run();
+	CHECK(StateGraph->GetStatus() == FStateGraph::EStatus::Waiting);
+	CHECK(TestNodeA->GetStatus() == FStateGraphNode::EStatus::Started);
+
+	Complete();
+	CHECK(StateGraph->GetStatus() == FStateGraph::EStatus::Completed);
+	CHECK(TestNodeA->GetStatus() == FStateGraphNode::EStatus::Completed);
+}
+
+class FClassNode : public FStateGraphNode
+{
+public:
+	FClassNode(FName InName) : FStateGraphNode(InName) {}
+	virtual void Start() { Complete(); }
+};
+
+void StaticNode(FStateGraphNodeFunctionComplete Complete) { Complete(); }
+
+class FRawClass
+{
+public:
+	void RawNode(FStateGraphNodeFunctionComplete Complete) { Complete(); }
+};
+
+class FSPClass : public TSharedFromThis<FSPClass>
+{
+public:
+	void SPNode(FStateGraphNodeFunctionComplete Complete) { Complete(); }
+};
+
+TEST_CASE("CreateNode for each supported type", "[CreateNodeFunctions]")
+{
+	FStateGraphRef StateGraph(MakeShared<FStateGraph>("Test"));
+	FRawClass RawClass;
+	TSharedRef<FSPClass> SPClass(MakeShared<FSPClass>());
+
+	StateGraph->CreateNode<FClassNode>("TestClass")
+		->Next("TestLambda", [](FStateGraphNodeFunctionComplete Complete) { Complete(); })
+		->Next("TestStatic", &StaticNode)
+		->Next("TestRaw", &RawClass, &FRawClass::RawNode)
+		->Next("TestSP", &SPClass.Get(), &FSPClass::SPNode)
+		->Next("TestSPLambda", &SPClass.Get(), [](FStateGraphNodeFunctionComplete Complete) { Complete(); });
+
+	StateGraph->Run();
+	CHECK(StateGraph->GetStatus() == FStateGraph::EStatus::Completed);
+}
+
+TEST_CASE("Node dependencies", "[NodeDependencies]")
+{
+	FStateGraphRef StateGraph(MakeShared<FStateGraph>("Test"));
+
+	FStateGraphNodeRef TestNodeA = StateGraph->CreateNode<FTestNode>("TestNodeA");
+	FStateGraphNodeRef TestNodeB = StateGraph->CreateNode<FTestNode>("TestNodeB", TSet<FName>({"TestNodeA"}));
+	FStateGraphNodeRef TestNodeC = StateGraph->CreateNode<FTestNode>("TestNodeC", TSet<FName>({"TestNodeB", "TestNodeD"}));
+
+	StateGraph->Run();
+	CHECK(StateGraph->GetStatus() == FStateGraph::EStatus::Waiting);
+	CHECK(TestNodeA->GetStatus() == FStateGraphNode::EStatus::Started);
+	CHECK(TestNodeB->GetStatus() == FStateGraphNode::EStatus::Blocked);
+	CHECK(TestNodeC->GetStatus() == FStateGraphNode::EStatus::Blocked);
+
+	TestNodeA->Complete();
+	CHECK(StateGraph->GetStatus() == FStateGraph::EStatus::Waiting);
+	CHECK(TestNodeA->GetStatus() == FStateGraphNode::EStatus::Completed);
+	CHECK(TestNodeB->GetStatus() == FStateGraphNode::EStatus::Started);
+	CHECK(TestNodeC->GetStatus() == FStateGraphNode::EStatus::Blocked);
+
+	TestNodeB->Complete();
+	CHECK(StateGraph->GetStatus() == FStateGraph::EStatus::Blocked);
+	CHECK(TestNodeB->GetStatus() == FStateGraphNode::EStatus::Completed);
+	CHECK(TestNodeC->GetStatus() == FStateGraphNode::EStatus::Blocked);
+
+	FStateGraphNodeRef TestNodeD = StateGraph->CreateNode<FTestNode>("TestNodeD");
+	StateGraph->Run();
+	CHECK(StateGraph->GetStatus() == FStateGraph::EStatus::Waiting);
+	CHECK(TestNodeC->GetStatus() == FStateGraphNode::EStatus::Blocked);
+	CHECK(TestNodeD->GetStatus() == FStateGraphNode::EStatus::Started);
+
+	TestNodeD->Complete();
+	CHECK(StateGraph->GetStatus() == FStateGraph::EStatus::Waiting);
+	CHECK(TestNodeC->GetStatus() == FStateGraphNode::EStatus::Started);
+	CHECK(TestNodeD->GetStatus() == FStateGraphNode::EStatus::Completed);
+
+	TestNodeC->Complete();
+	CHECK(StateGraph->GetStatus() == FStateGraph::EStatus::Completed);
+	CHECK(TestNodeC->GetStatus() == FStateGraphNode::EStatus::Completed);
+}
+
+TEST_CASE("Adding and reusing nodes", "[AddingNodes]")
+{
+	FStateGraphRef StateGraphA(MakeShared<FStateGraph>("TestB"));
+	FStateGraphNodeRef TestNodeA = MakeShared<FTestNode>("TestNodeA");
+
+	CHECK(StateGraphA->AddNode(TestNodeA));
+	StateGraphA->Run();
+	CHECK(StateGraphA->GetStatus() == FStateGraph::EStatus::Waiting);
+	CHECK(TestNodeA->GetStatus() == FStateGraphNode::EStatus::Started);
+
+	StateGraphA->RemoveNode(TestNodeA->GetName());
+	StateGraphA->Run();
+	CHECK(StateGraphA->GetStatus() == FStateGraph::EStatus::Completed);
+	CHECK(TestNodeA->GetStatus() == FStateGraphNode::EStatus::Started);
+
+	FStateGraphRef StateGraphB(MakeShared<FStateGraph>("TestB"));
+	CHECK(StateGraphB->AddNode(TestNodeA));
+	StateGraphB->Run();
+	CHECK(StateGraphB->GetStatus() == FStateGraph::EStatus::Waiting);
+	CHECK(TestNodeA->GetStatus() == FStateGraphNode::EStatus::Started);
+
+	TestNodeA->Complete();
+	CHECK(StateGraphB->GetStatus() == FStateGraph::EStatus::Completed);
+	CHECK(TestNodeA->GetStatus() == FStateGraphNode::EStatus::Completed);
+
+	StateGraphA->Reset();
+	CHECK(StateGraphA->GetStatus() == FStateGraph::EStatus::NotStarted);
+	CHECK(!StateGraphA->AddNode(TestNodeA));
+
+	StateGraphB->RemoveNode(TestNodeA->GetName());
+	CHECK(StateGraphA->AddNode(TestNodeA));
+	StateGraphA->Run();
+	CHECK(StateGraphA->GetStatus() == FStateGraph::EStatus::Completed);
+	CHECK(TestNodeA->GetStatus() == FStateGraphNode::EStatus::Completed);
+
+	FStateGraphNodeRef TestNodeA2 = MakeShared<FTestNode>("TestNodeA");
+	CHECK(!StateGraphA->AddNode(TestNodeA2));
+}
+
+TEST_CASE("Removing nodes", "[RemoveNodes]")
+{
+	FStateGraphRef StateGraph(MakeShared<FStateGraph>("Test"));
+
+	FTestNodePtr TestNodeA = StateGraph->CreateNode<FTestNode>("TestNodeA");
+	FTestNodeWeakPtr WeakTestNodeA = TestNodeA;
+	StateGraph->RemoveNode("TestNodeA");
+	CHECK(TestNodeA->GetStatus() == FStateGraphNode::EStatus::NotStarted);
+	CHECK(TestNodeA->bWasRemovedCalled);
+	TestNodeA.Reset();
+	CHECK(!WeakTestNodeA.IsValid());
+
+	WeakTestNodeA = StateGraph->CreateNode<FTestNode>("TestNodeA");
+	StateGraph->RemoveNode("TestNodeA");
+	CHECK(!WeakTestNodeA.IsValid());
+
+	WeakTestNodeA = StateGraph->CreateNode<FTestNode>("TestNodeA");
+	FStateGraphNodePtr TestNodeB = StateGraph->CreateNode("TestNodeB", [StateGraph](FStateGraphNodeFunctionComplete Complete)
+	{
+		StateGraph->RemoveNode("TestNodeA");
+		StateGraph->RemoveNode("TestNodeC");
+	});
+	FStateGraphNodeWeakPtr WeakTestNodeC = StateGraph->CreateNode<FTestNode>("TestNodeC");
+	StateGraph->Run();
+	CHECK(!WeakTestNodeA.IsValid());
+	CHECK(!WeakTestNodeC.IsValid());
+	StateGraph->RemoveNode("TestNodeB");
+
+	StateGraph->CreateNode<FTestNode>("TestNodeA")
+		->Next<FTestNode>("TestNodeB")
+		->Next<FTestNode>("TestNodeC");
+	StateGraph->RemoveAllNodes();
+	CHECK(StateGraph->GetNodeRef("TestNodeA") == nullptr);
+	CHECK(StateGraph->GetNodeRef("TestNodeB") == nullptr);
+	CHECK(StateGraph->GetNodeRef("TestNodeC") == nullptr);
+}
+
+TEST_CASE("Resetting nodes", "[ResetNodes]")
+{
+	FStateGraphRef StateGraph(MakeShared<FStateGraph>("Test"));
+
+	int32 StartCount = 0;
+	FStateGraphNodeRef TestNodeA = StateGraph->CreateNode("TestNodeA", [&StartCount, &TestNodeA](FStateGraphNodeFunctionComplete Complete)
+	{
+		CHECK(TestNodeA->GetStatus() == FStateGraphNode::EStatus::Started);
+
+		if (StartCount == 0)
+		{
+			TestNodeA->Reset();
+		}
+		else
+		{
+			Complete();
+		}
+
+		++StartCount;
+	});
+
+	StateGraph->Run();
+	CHECK(StartCount == 1);
+	CHECK(TestNodeA->GetStatus() == FStateGraphNode::EStatus::NotStarted);
+
+	StateGraph->Run();
+	CHECK(StartCount == 2);
+	CHECK(TestNodeA->GetStatus() == FStateGraphNode::EStatus::Completed);
+
+	TestNodeA->Reset();
+	StateGraph->Run();
+	CHECK(StartCount == 3);
+}
+
+TEST_CASE("Resetting state graph", "[ResetStateGraph]")
+{
+	FStateGraphRef StateGraph(MakeShared<FStateGraph>("Test"));
+
+	int32 StartCount = 0;
+	FStateGraphNodeRef TestNodeA = StateGraph->CreateNode("TestNodeA", [&StartCount, &StateGraph](FStateGraphNodeFunctionComplete Complete)
+	{
+		if (StartCount == 0)
+		{
+			StateGraph->Reset();
+		}
+		else
+		{
+			Complete();
+		}
+
+		++StartCount;
+	});
+
+	StateGraph->Run();
+	CHECK(StartCount == 1);
+	CHECK(TestNodeA->GetStatus() == FStateGraphNode::EStatus::NotStarted);
+	CHECK(StateGraph->GetStatus() == FStateGraph::EStatus::NotStarted);
+
+	StateGraph->Run();
+	CHECK(StartCount == 2);
+	CHECK(TestNodeA->GetStatus() == FStateGraphNode::EStatus::Completed);
+	CHECK(StateGraph->GetStatus() == FStateGraph::EStatus::Completed);
+}
+
+TEST_CASE("Deleting state graph", "[DeleteStateGraph]")
+{
+	FStateGraphPtr StateGraph(MakeShared<FStateGraph>("Test"));
+	FStateGraphWeakPtr WeakStateGraph(StateGraph);
+	FStateGraphNodeRef TestNodeA = StateGraph->CreateNode("TestNodeA", [&StateGraph, &WeakStateGraph](FStateGraphNodeFunctionComplete Complete)
+	{
+		StateGraph.Reset();
+		CHECK(!WeakStateGraph.IsValid());
+	});
+
+	StateGraph->Run();
+	CHECK(!WeakStateGraph.IsValid());
+	CHECK(TestNodeA->GetStateGraphName() == NAME_None);
+}
+
+TEST_CASE("Pausing state graph", "[PauseStateGraph]")
+{
+	FStateGraphRef StateGraph(MakeShared<FStateGraph>("Test"));
+	FStateGraphNodeRef TestNodeA = StateGraph->CreateNode("TestNodeA", [&StateGraph](FStateGraphNodeFunctionComplete Complete) { StateGraph->Pause(); });
+	FStateGraphNodeRef TestNodeB = StateGraph->CreateNode<FTestNode>("TestNodeB", TSet<FName>({"TestNodeA"}));
+	FStateGraphNodeRef TestNodeC = StateGraph->CreateNode("TestNodeC", [](FStateGraphNodeFunctionComplete Complete) { Complete(); });
+	TestNodeC->Dependencies.Add("TestNodeB");
+
+	StateGraph->Run();
+	CHECK(StateGraph->GetStatus() == FStateGraph::EStatus::Paused);
+	CHECK(TestNodeA->GetStatus() == FStateGraphNode::EStatus::Started);
+
+	TestNodeA->Complete();
+	CHECK(StateGraph->GetStatus() == FStateGraph::EStatus::Paused);
+	CHECK(TestNodeA->GetStatus() == FStateGraphNode::EStatus::Completed);
+	CHECK(TestNodeB->GetStatus() != FStateGraphNode::EStatus::Started);
+
+	StateGraph->Run();
+	CHECK(StateGraph->GetStatus() == FStateGraph::EStatus::Waiting);
+	CHECK(TestNodeA->GetStatus() == FStateGraphNode::EStatus::Completed);
+	CHECK(TestNodeB->GetStatus() == FStateGraphNode::EStatus::Started);
+	CHECK(TestNodeC->GetStatus() == FStateGraphNode::EStatus::Blocked);
+
+	StateGraph->Pause();
+	TestNodeB->Complete();
+	CHECK(StateGraph->GetStatus() == FStateGraph::EStatus::Paused);
+	CHECK(TestNodeB->GetStatus() == FStateGraphNode::EStatus::Completed);
+	CHECK(TestNodeC->GetStatus() == FStateGraphNode::EStatus::Blocked);
+
+	StateGraph->Run();
+	CHECK(StateGraph->GetStatus() == FStateGraph::EStatus::Completed);
+	CHECK(TestNodeC->GetStatus() == FStateGraphNode::EStatus::Completed);
+}
+
+TEST_CASE("State graph timeout", "[StateGraphTimeout]")
+{
+	FStateGraphRef StateGraph(MakeShared<FStateGraph>("Test"));
+	StateGraph->SetTimeout(0.01);
+	bool bOnTimedOutCalled = false;
+	StateGraph->OnStatusChanged.AddLambda([&bOnTimedOutCalled](FStateGraph&, FStateGraph::EStatus OldStatus, FStateGraph::EStatus NewStatus)
+	{
+		if (NewStatus == FStateGraph::EStatus::TimedOut)
+		{
+			bOnTimedOutCalled = true;
+		}
+	});
+
+	FStateGraphNodeRef TestNodeA = StateGraph->CreateNode("TestNodeA", [](FStateGraphNodeFunctionComplete Complete) {});
+	StateGraph->Run();
+
+	while (StateGraph->GetStatus() == FStateGraph::EStatus::Waiting)
+	{
+		FTSTicker::GetCoreTicker().Tick(0.1f);
+	}
+
+	CHECK(StateGraph->GetStatus() == FStateGraph::EStatus::TimedOut);
+	CHECK(bOnTimedOutCalled);
+}
+
+TEST_CASE("Node timeout", "[NodeTimeout]")
+{
+	FStateGraphRef StateGraph(MakeShared<FStateGraph>("Test"));
+	bool bOnNodeTimedOutCalled = false;
+	StateGraph->OnNodeStatusChanged.AddLambda([&bOnNodeTimedOutCalled](FStateGraphNode&, FStateGraphNode::EStatus OldStatus, FStateGraphNode::EStatus NewStatus)
+	{
+		if (NewStatus == FStateGraphNode::EStatus::TimedOut)
+		{
+			bOnNodeTimedOutCalled = true;
+		}
+	});
+
+	FTestNodeRef TestNodeA = StateGraph->CreateNode<FTestNode>("TestNodeA");
+	TestNodeA->SetTimeout(0.01);
+	StateGraph->Run();
+
+	while (StateGraph->GetStatus() == FStateGraph::EStatus::Waiting)
+	{
+		FTSTicker::GetCoreTicker().Tick(0.1f);
+	}
+
+	CHECK(TestNodeA->GetStatus() == FStateGraphNode::EStatus::TimedOut);
+	CHECK(StateGraph->GetStatus() == FStateGraph::EStatus::Blocked);
+	CHECK(TestNodeA->bWasTimedOutCalled);
+	CHECK(bOnNodeTimedOutCalled);
+
+	TestNodeA->Complete();
+	CHECK(StateGraph->GetStatus() == FStateGraph::EStatus::Completed);
+}
+
+TEST_CASE("State graph config", "[StateGraphConfig]")
+{
+	FConfigCacheIni::InitializeConfigSystem();
+	FConfigFile* EngineIni = GConfig->FindConfigFile(GEngineIni);
+	check(EngineIni);
+	EngineIni->CombineFromBuffer(TEXT("[StateGraph.Test]\nTimeout=0.01\n"), GEngineIni);
+
+	FStateGraphRef StateGraph(MakeShared<FStateGraph>("Test"));
+	StateGraph->Initialize();
+	FStateGraphNodeRef TestNodeA = StateGraph->CreateNode("TestNodeA", [](FStateGraphNodeFunctionComplete Complete) {});
+	StateGraph->Run();
+
+	while (StateGraph->GetStatus() == FStateGraph::EStatus::Waiting)
+	{
+		FTSTicker::GetCoreTicker().Tick(0.1f);
+	}
+
+	CHECK(StateGraph->GetStatus() == FStateGraph::EStatus::TimedOut);
+	TestNodeA->Complete();
+	CHECK(StateGraph->GetStatus() == FStateGraph::EStatus::TimedOut);
+
+	EngineIni->CombineFromBuffer(TEXT("[StateGraph.Test]\nTimeout=0\n"), GEngineIni);
+	FCoreDelegates::TSOnConfigSectionsChanged().Broadcast(GEngineIni, TSet<FString>({ TEXT("StateGraph.Test") }));
+
+	StateGraph->Run();
+	CHECK(StateGraph->GetStatus() == FStateGraph::EStatus::Completed);
+}
+
+TEST_CASE("Node config", "[NodeConfig]")
+{
+	FConfigCacheIni::InitializeConfigSystem();
+	FConfigFile* EngineIni = GConfig->FindConfigFile(GEngineIni);
+	check(EngineIni);
+	EngineIni->CombineFromBuffer(TEXT("[StateGraph.Test.TestNodeA]\nTimeout=0.01\nCustomConfig=2\n"), GEngineIni);
+
+	FStateGraphRef StateGraph(MakeShared<FStateGraph>("Test"));
+	StateGraph->Initialize();
+	FTestNodeRef TestNodeA = StateGraph->CreateNode<FTestNode>("TestNodeA");
+	StateGraph->Run();
+
+	while (StateGraph->GetStatus() == FStateGraph::EStatus::Waiting)
+	{
+		FTSTicker::GetCoreTicker().Tick(0.1f);
+	}
+
+	CHECK(TestNodeA->GetStatus() == FStateGraphNode::EStatus::TimedOut);
+	CHECK(StateGraph->GetStatus() == FStateGraph::EStatus::Blocked);
+	CHECK(TestNodeA->bWasTimedOutCalled);
+	CHECK(TestNodeA->CustomConfig == 2);
+
+	EngineIni->CombineFromBuffer(TEXT("[StateGraph.Test.TestNodeA]\nTimeout=0.01\nCustomConfig=3\n"), GEngineIni);
+	FCoreDelegates::TSOnConfigSectionsChanged().Broadcast(GEngineIni, TSet<FString>({ TEXT("StateGraph.Test.TestNodeA") }));
+
+	CHECK(TestNodeA->CustomConfig == 3);
+}
+
+} // UE::StateGraphTests
diff --git a/Engine/Source/Programs/StateGraphTests/StateGraphTests.Build.cs b/Engine/Source/Programs/StateGraphTests/StateGraphTests.Build.cs
new file mode 100644
--- /dev/null
+++ b/Engine/Source/Programs/StateGraphTests/StateGraphTests.Build.cs
@@ -0,0 +1,19 @@
+// Copyright Epic Games, Inc. All Rights Reserved.
+
+using UnrealBuildTool;
+
+public class StateGraphTests : TestModuleRules
+{
+	public StateGraphTests(ReadOnlyTargetRules Target) : base(Target)
+	{
+		OptimizeCode = CodeOptimization.Never;
+
+		PrivateDependencyModuleNames.AddRange(
+			new string[] {
+				"Core",
+				"StateGraph"
+			});
+
+		UpdateBuildGraphPropertiesFile(new Metadata() { TestName = "StateGraph", TestShortName = "State Graph"});
+	}
+}
diff --git a/Engine/Source/Programs/StateGraphTests/StateGraphTests.Target.cs b/Engine/Source/Programs/StateGraphTests/StateGraphTests.Target.cs
new file mode 100644
--- /dev/null
+++ b/Engine/Source/Programs/StateGraphTests/StateGraphTests.Target.cs
@@ -0,0 +1,12 @@
+// Copyright Epic Games, Inc. All Rights Reserved.
+
+using UnrealBuildTool;
+
+public class StateGraphTestsTarget : TestTargetRules
+{
+	public StateGraphTestsTarget(TargetInfo Target) : base(Target)
+	{
+		bForceEnableExceptions = true;
+		EnablePlugins.Add("StateGraph");
+	}
+}
