Add missing header in StateGraph.

diff --git a/Engine/Plugins/Experimental/StateGraph/Source/StateGraph/Private/StateGraph.cpp b/Engine/Plugins/Experimental/StateGraph/Source/StateGraph/Private/StateGraph.cpp
--- a/Engine/Plugins/Experimental/StateGraph/Source/StateGraph/Private/StateGraph.cpp
+++ b/Engine/Plugins/Experimental/StateGraph/Source/StateGraph/Private/StateGraph.cpp
@@ -4,6 +4,7 @@
 
 #include "HAL/PlatformTime.h"
 #include "Misc/ConfigCacheIni.h"
+#include "Misc/CoreDelegates.h"
 #include "Misc/Parse.h"
 #include "Modules/ModuleManager.h"
 
