StateGraph: Add forward header and remove name from node construction messages.

diff --git a/Engine/Plugins/Experimental/StateGraph/Source/StateGraph/Private/StateGraph.cpp b/Engine/Plugins/Experimental/StateGraph/Source/StateGraph/Private/StateGraph.cpp
--- a/Engine/Plugins/Experimental/StateGraph/Source/StateGraph/Private/StateGraph.cpp
+++ b/Engine/Plugins/Experimental/StateGraph/Source/StateGraph/Private/StateGraph.cpp
@@ -23,12 +23,12 @@ namespace UE {
 FStateGraphNode::FStateGraphNode(FName InName) :
 	Name(InName)
 {
-	UE_LOG_STATEGRAPH(VeryVerbose, TEXT("[%s.%s] Created node"), *GetStateGraphName().ToString(), *Name.ToString());
+	UE_LOG_STATEGRAPH(VeryVerbose, TEXT("[%s] Created node"), *Name.ToString());
 }
 
 FStateGraphNode::~FStateGraphNode()
 {
-	UE_LOG_STATEGRAPH(VeryVerbose, TEXT("[%s.%s] Destroyed node"), *GetStateGraphName().ToString(), *Name.ToString());
+	UE_LOG_STATEGRAPH(VeryVerbose, TEXT("[%s] Destroyed node"), *Name.ToString());
 }
 
 const TCHAR* FStateGraphNode::GetStatusName(EStatus Status)
diff --git a/Engine/Plugins/Experimental/StateGraph/Source/StateGraph/Public/StateGraph.h b/Engine/Plugins/Experimental/StateGraph/Source/StateGraph/Public/StateGraph.h
--- a/Engine/Plugins/Experimental/StateGraph/Source/StateGraph/Public/StateGraph.h
+++ b/Engine/Plugins/Experimental/StateGraph/Source/StateGraph/Public/StateGraph.h
@@ -19,9 +19,8 @@
 #include "Delegates/Delegate.h"
 #include "Logging/LogMacros.h"
 #include "Misc/CoreMisc.h"
-#include "Templates/Function.h"
+#include "StateGraphFwd.h"
 #include "Templates/IsClass.h"
-#include "Templates/SharedPointer.h"
 #include "Templates/UnrealTemplate.h"
 #include "UObject/NameTypes.h"
 
@@ -30,21 +29,6 @@ STATEGRAPH_API DECLARE_LOG_CATEGORY_EXTERN(LogStateGraph, Log, All);
 
 namespace UE {
 
-class FStateGraph;
-using FStateGraphRef = TSharedRef<FStateGraph, ESPMode::ThreadSafe>;
-using FStateGraphPtr = TSharedPtr<FStateGraph, ESPMode::ThreadSafe>;
-using FStateGraphWeakPtr = TWeakPtr<FStateGraph, ESPMode::ThreadSafe>;
-
-class FStateGraphNode;
-using FStateGraphNodeRef = TSharedRef<FStateGraphNode, ESPMode::ThreadSafe>;
-using FStateGraphNodePtr = TSharedPtr<FStateGraphNode, ESPMode::ThreadSafe>;
-using FStateGraphNodeWeakPtr = TWeakPtr<FStateGraphNode, ESPMode::ThreadSafe>;
-
-class FStateGraphNodeFunction;
-using FStateGraphNodeFunctionRef = TSharedRef<FStateGraphNodeFunction, ESPMode::ThreadSafe>;
-using FStateGraphNodeFunctionPtr = TSharedPtr<FStateGraphNodeFunction, ESPMode::ThreadSafe>;
-using FStateGraphNodeFunctionWeakPtr = TWeakPtr<FStateGraphNodeFunction, ESPMode::ThreadSafe>;
-
 /**
  * Abstract base class for all nodes in the state graph.
  */
@@ -173,9 +157,6 @@ private:
 	friend FStateGraph;
 };
 
-/** Complete function type passed into FStateGraphNodeFunction node start functions. */
-using FStateGraphNodeFunctionComplete = TFunction<void()>;
-
 /** Start function delegate type for FStateGraphNodeFunction nodes. */
 DECLARE_DELEGATE_OneParam(FStateGraphNodeFunctionStart, FStateGraphNodeFunctionComplete /*Complete*/);
 
diff --git a/Engine/Plugins/Experimental/StateGraph/Source/StateGraph/Public/StateGraphFwd.h b/Engine/Plugins/Experimental/StateGraph/Source/StateGraph/Public/StateGraphFwd.h
new file mode 100644
--- /dev/null
+++ b/Engine/Plugins/Experimental/StateGraph/Source/StateGraph/Public/StateGraphFwd.h
@@ -0,0 +1,28 @@
+// Copyright Epic Games, Inc. All Rights Reserved.
+
+#pragma once
+
+#include "Templates/Function.h"
+#include "Templates/SharedPointer.h"
+
+namespace UE {
+
+class FStateGraph;
+using FStateGraphRef = TSharedRef<FStateGraph, ESPMode::ThreadSafe>;
+using FStateGraphPtr = TSharedPtr<FStateGraph, ESPMode::ThreadSafe>;
+using FStateGraphWeakPtr = TWeakPtr<FStateGraph, ESPMode::ThreadSafe>;
+
+class FStateGraphNode;
+using FStateGraphNodeRef = TSharedRef<FStateGraphNode, ESPMode::ThreadSafe>;
+using FStateGraphNodePtr = TSharedPtr<FStateGraphNode, ESPMode::ThreadSafe>;
+using FStateGraphNodeWeakPtr = TWeakPtr<FStateGraphNode, ESPMode::ThreadSafe>;
+
+class FStateGraphNodeFunction;
+using FStateGraphNodeFunctionRef = TSharedRef<FStateGraphNodeFunction, ESPMode::ThreadSafe>;
+using FStateGraphNodeFunctionPtr = TSharedPtr<FStateGraphNodeFunction, ESPMode::ThreadSafe>;
+using FStateGraphNodeFunctionWeakPtr = TWeakPtr<FStateGraphNodeFunction, ESPMode::ThreadSafe>;
+
+/** Complete function type passed into FStateGraphNodeFunction node start functions. */
+using FStateGraphNodeFunctionComplete = TFunction<void()>;
+
+} // UE
