StateGraph: Small logging count fix and add extra timeout logging.

diff --git a/Engine/Plugins/Experimental/StateGraph/Source/StateGraph/Private/StateGraph.cpp b/Engine/Plugins/Experimental/StateGraph/Source/StateGraph/Private/StateGraph.cpp
--- a/Engine/Plugins/Experimental/StateGraph/Source/StateGraph/Private/StateGraph.cpp
+++ b/Engine/Plugins/Experimental/StateGraph/Source/StateGraph/Private/StateGraph.cpp
@@ -89,7 +89,7 @@ void FStateGraphNode::Complete()
 	}
 
 	CompletedTime = FPlatformTime::Seconds();
-	UE_LOG_STATEGRAPH(Log, TEXT("[%s.%s] Completed node (Duration=%.6f)"), *GetStateGraphName().ToString(), *Name.ToString(), CompletedTime - StartTime);
+	UE_LOG_STATEGRAPH(Log, TEXT("[%s.%s] Completed node (Duration=%.6f Timeout=%.6f)"), *GetStateGraphName().ToString(), *Name.ToString(), CompletedTime - StartTime, Timeout);
 
 	// Keep a reference to check if the node is destroyed during external functions.
 	FStateGraphNodeWeakPtr StateGraphNodeWeakPtr(AsWeak());
@@ -336,7 +336,7 @@ void FStateGraph::Run()
 		NextTimeout = (StartTime + Timeout) - Now;
 		if (NextTimeout <= 0.f)
 		{
-			UE_LOG_STATEGRAPH(Log, TEXT("[%s] State graph timed out (Duration=%.6f)"), *Name.ToString(), Now - StartTime);
+			UE_LOG_STATEGRAPH(Log, TEXT("[%s] State graph timed out (Duration=%.6f Timeout=%.6f)"), *Name.ToString(), Now - StartTime, Timeout);
 			SetStatus(EStatus::TimedOut);
 			return;
 		}
@@ -344,7 +344,7 @@ void FStateGraph::Run()
 
 	bRunning = true;
 
-	UE_LOG_STATEGRAPH(Verbose, TEXT("[%s] Starting run loop"), *Name.ToString());
+	UE_LOG_STATEGRAPH(Verbose, TEXT("[%s] Starting run loop (Now=%.06f)"), *Name.ToString(), Now);
 
 	uint32 Blocked = 0;
 	uint32 Started = 0;
@@ -371,6 +371,8 @@ void FStateGraph::Run()
 			continue;
 		}
 
+		bool bCounted = false;
+
 		switch ((*Node)->Status)
 		{
 		case FStateGraphNode::EStatus::NotStarted:
@@ -450,7 +452,8 @@ void FStateGraph::Run()
 				break;
 			}
 
-			// Fall through to checkout for timeout.
+			bCounted = true;
+			// Fall through to check for timeout.
 
 		case FStateGraphNode::EStatus::Started:
 			if ((*Node)->Timeout > 0.f)
@@ -459,7 +462,7 @@ void FStateGraph::Run()
 				if (NodeTimeout <= 0.f)
 				{
 					++TimedOut;
-					UE_LOG_STATEGRAPH(Log, TEXT("[%s.%s] Node timed out (Duration=%.6f)"), *Name.ToString(), *NodeName.ToString(), Now - (*Node)->StartTime);
+					UE_LOG_STATEGRAPH(Log, TEXT("[%s.%s] Node timed out (Duration=%.6f Timeout=%.6f)"), *Name.ToString(), *NodeName.ToString(), Now - (*Node)->StartTime, (*Node)->Timeout);
 					(*Node)->SetStatus(FStateGraphNode::EStatus::TimedOut);
 
 					if (!StateGraphWeakPtr.IsValid())
@@ -486,7 +489,11 @@ void FStateGraph::Run()
 				}
 			}
 
-			++Running;
+			if (!bCounted)
+			{
+				++Running;
+			}
+
 			break;
 
 		case FStateGraphNode::EStatus::Completed:
@@ -505,8 +512,8 @@ void FStateGraph::Run()
 	}
 
 	bRunning = false;
-	UE_LOG_STATEGRAPH(Verbose, TEXT("[%s] Duration=%.6f Blocked=%d Started=%d Running=%d Completed=%d Removed=%d TimedOut=%d"),
-		*Name.ToString(), Now - StartTime, Blocked, Started, Running, Completed, Removed, TimedOut);
+	UE_LOG_STATEGRAPH(Verbose, TEXT("[%s] Duration=%.6f Timeout=%.6f Blocked=%d Started=%d Running=%d Completed=%d Removed=%d TimedOut=%d"),
+		*Name.ToString(), Now - StartTime, Timeout, Blocked, Started, Running, Completed, Removed, TimedOut);
 
 	if (bRunAgain)
 	{
@@ -525,7 +532,7 @@ void FStateGraph::Run()
 		if (Blocked == 0 && TimedOut == 0)
 		{
 			CompletedTime = Now;
-			UE_LOG_STATEGRAPH(Verbose, TEXT("[%s] Completed (Duration=%.6f)"), *Name.ToString(), CompletedTime - StartTime);
+			UE_LOG_STATEGRAPH(Log, TEXT("[%s] Completed (Duration=%.6f Timeout=%.6f)"), *Name.ToString(), CompletedTime - StartTime, Timeout);
 			SetStatus(EStatus::Completed);
 		}
 		else
@@ -548,6 +555,7 @@ void FStateGraph::Run()
 
 	if ((Status == EStatus::Blocked || Status == EStatus::Waiting) && NextTimeout > 0.f)
 	{
+		UE_LOG_STATEGRAPH(Verbose, TEXT("[%s] Setting timer for %.6f"), *Name.ToString(), NextTimeout);
 		TimeoutTicker = FTSTicker::GetCoreTicker().AddTicker(FTickerDelegate::CreateSPLambda(this, [this](float DeltaTime) {
 			Run();
 			return false;
@@ -603,7 +611,7 @@ void FStateGraph::LogDebugInfo(bool bWarning)
 		return;
 	}
 
-	FString Message = FString::Printf(TEXT("[%s] Status=%s Nodes=%d Duration=%.6f"), *Name.ToString(), GetStatusName(), Nodes.Num(), GetDuration());
+	FString Message = FString::Printf(TEXT("[%s] Status=%s Nodes=%d Duration=%.6f Timeout=%.6f"), *Name.ToString(), GetStatusName(), Nodes.Num(), GetDuration(), Timeout);
 
 	if (bWarning)
 	{
@@ -667,8 +675,8 @@ void FStateGraph::LogDebugInfo(bool bWarning)
 			Dependencies.Add(TEXT("None"));
 		}
 
-		Message = FString::Printf(TEXT("[%s.%s] Status=%s Duration=%.6f Dependencies(%s)"),
-			*Name.ToString(), *Node.Key.ToString(), Node.Value->GetStatusName(), Node.Value->GetDuration(), *FString::Join(Dependencies, TEXT(" ")));
+		Message = FString::Printf(TEXT("[%s.%s] Status=%s Duration=%.6f Timeout=%.6f Dependencies(%s)"),
+			*Name.ToString(), *Node.Key.ToString(), Node.Value->GetStatusName(), Node.Value->GetDuration(), Node.Value->Timeout, *FString::Join(Dependencies, TEXT(" ")));
 
 		if (bWarning)
 		{
