Fix status grouping in state graph debug logging.

diff --git a/Engine/Plugins/Experimental/StateGraph/Source/StateGraph/Private/StateGraph.cpp b/Engine/Plugins/Experimental/StateGraph/Source/StateGraph/Private/StateGraph.cpp
--- a/Engine/Plugins/Experimental/StateGraph/Source/StateGraph/Private/StateGraph.cpp
+++ b/Engine/Plugins/Experimental/StateGraph/Source/StateGraph/Private/StateGraph.cpp
@@ -624,28 +624,15 @@ void FStateGraph::LogDebugInfo(bool bWarning)
 
 	for (auto Node : Nodes)
 	{
+		TMap<FStateGraphNode::EStatus, TArray<FString>> DependenciesByStatus;
 		TArray<FString> Missing;
-		TArray<FString> Blocked;
-		TArray<FString> Completed;
-		TArray<FString> TimedOut;
 
 		for (FName Dependency : Node.Value->Dependencies)
 		{
 			FStateGraphNodeRef* DependencyNode = GetNodeRef(Dependency);
 			if (DependencyNode)
 			{
-				if ((*DependencyNode)->Status == FStateGraphNode::EStatus::Completed)
-				{
-					Completed.Add(Dependency.ToString());
-				}
-				else if ((*DependencyNode)->Status == FStateGraphNode::EStatus::Blocked)
-				{
-					Blocked.Add(Dependency.ToString());
-				}
-				else
-				{
-					TimedOut.Add(Dependency.ToString());
-				}
+				DependenciesByStatus.FindOrAdd((*DependencyNode)->Status).Add(Dependency.ToString());
 			}
 			else
 			{
@@ -658,18 +645,12 @@ void FStateGraph::LogDebugInfo(bool bWarning)
 		{
 			Dependencies.Add(TEXT("Missing=") + FString::Join(Missing, TEXT(",")));
 		}
-		if (Blocked.Num())
-		{
-			Dependencies.Add(TEXT("Blocked=") + FString::Join(Blocked, TEXT(",")));
-		}
-		if (Completed.Num())
-		{
-			Dependencies.Add(TEXT("Completed=") + FString::Join(Completed, TEXT(",")));
-		}
-		if (TimedOut.Num())
+
+		for (const TPair<FStateGraphNode::EStatus, TArray<FString>>& Pair : DependenciesByStatus)
 		{
-			Dependencies.Add(TEXT("TimedOut=") + FString::Join(TimedOut, TEXT(",")));
+			Dependencies.Add(FString::Printf(TEXT("%s=%s"), FStateGraphNode::GetStatusName(Pair.Key), *FString::Join(Pair.Value, TEXT(","))));
 		}
+
 		if (Dependencies.Num() == 0)
 		{
 			Dependencies.Add(TEXT("None"));
