Downgrade Warning to Log. This flow is supported.

diff --git a/Engine/Plugins/Experimental/StateGraph/Source/StateGraph/Private/StateGraph.cpp b/Engine/Plugins/Experimental/StateGraph/Source/StateGraph/Private/StateGraph.cpp
--- a/Engine/Plugins/Experimental/StateGraph/Source/StateGraph/Private/StateGraph.cpp
+++ b/Engine/Plugins/Experimental/StateGraph/Source/StateGraph/Private/StateGraph.cpp
@@ -180,7 +180,7 @@ void FStateGraphNodeFunction::Start()
 		}
 		else
 		{
-			UE_LOG_STATEGRAPH(Warning, TEXT("[%s.%s] Function node completed after node was destroyed"), *StateGraphName.ToString(), *NodeName.ToString());
+			UE_LOG_STATEGRAPH(Log, TEXT("[%s.%s] Function node completed after node was destroyed"), *StateGraphName.ToString(), *NodeName.ToString());
 		}
 	});
 }
