StateGraph updates:
- Pass a reference into function node delegates.
- Add context name for better logging when running concurrent graphs of the same type.
#rnx

diff --git a/Engine/Plugins/Experimental/StateGraph/Source/StateGraph/Private/StateGraph.cpp b/Engine/Plugins/Experimental/StateGraph/Source/StateGraph/Private/StateGraph.cpp
--- a/Engine/Plugins/Experimental/StateGraph/Source/StateGraph/Private/StateGraph.cpp
+++ b/Engine/Plugins/Experimental/StateGraph/Source/StateGraph/Private/StateGraph.cpp
@@ -21,14 +21,15 @@ namespace UE {
 // FStateGraphNode
 
 FStateGraphNode::FStateGraphNode(FName InName) :
-	Name(InName)
+	Name(InName),
+	ContextName(Name.ToString())
 {
-	UE_LOG_STATEGRAPH(VeryVerbose, TEXT("[%s] Created node"), *Name.ToString());
+	UE_LOG_STATEGRAPH(VeryVerbose, TEXT("[%s] Created node"), *ContextName);
 }
 
 FStateGraphNode::~FStateGraphNode()
 {
-	UE_LOG_STATEGRAPH(VeryVerbose, TEXT("[%s] Destroyed node"), *Name.ToString());
+	UE_LOG_STATEGRAPH(VeryVerbose, TEXT("[%s] Destroyed node"), *ContextName);
 }
 
 const TCHAR* FStateGraphNode::GetStatusName(EStatus Status)
@@ -64,7 +65,7 @@ bool FStateGraphNode::CheckDependencies() const
 	FStateGraphPtr StateGraphPtr(StateGraphWeakPtr.Pin());
 	if (!StateGraphPtr.IsValid())
 	{
-		UE_LOG_STATEGRAPH(Warning, TEXT("[%s] Node checked with invalid state graph"), *Name.ToString());
+		UE_LOG_STATEGRAPH(Warning, TEXT("[%s] Node checked with invalid state graph"), *ContextName);
 		return false;
 	}
 
@@ -84,12 +85,12 @@ void FStateGraphNode::Complete()
 {
 	if (Status == EStatus::Completed)
 	{
-		UE_LOG_STATEGRAPH(Warning, TEXT("[%s.%s] Node already completed"), *GetStateGraphName().ToString(), *Name.ToString());
+		UE_LOG_STATEGRAPH(Warning, TEXT("[%s] Node already completed"), *ContextName);
 		return;
 	}
 
 	CompletedTime = FPlatformTime::Seconds();
-	UE_LOG_STATEGRAPH(Log, TEXT("[%s.%s] Completed node (Duration=%.6f Timeout=%.6f)"), *GetStateGraphName().ToString(), *Name.ToString(), CompletedTime - StartTime, Timeout);
+	UE_LOG_STATEGRAPH(Log, TEXT("[%s] Completed node (Duration=%.6f Timeout=%.6f)"), *ContextName, CompletedTime - StartTime, Timeout);
 
 	// Keep a reference to check if the node is destroyed during external functions.
 	FStateGraphNodeWeakPtr StateGraphNodeWeakPtr(AsWeak());
@@ -112,23 +113,18 @@ void FStateGraphNode::Complete()
 	}
 	else
 	{
-		UE_LOG_STATEGRAPH(Warning, TEXT("[%s] Node completed with invalid state graph"), *Name.ToString());
+		UE_LOG_STATEGRAPH(Warning, TEXT("[%s] Node completed with invalid state graph"), *ContextName);
 	}
 }
 
 void FStateGraphNode::Reset()
 {
-	UE_LOG_STATEGRAPH(Verbose, TEXT("[%s.%s] Resetting node"), *GetStateGraphName().ToString(), *Name.ToString());
+	UE_LOG_STATEGRAPH(Verbose, TEXT("[%s] Resetting node"), *ContextName);
 	StartTime = 0.f;
 	CompletedTime = 0.f;
 	SetStatus(EStatus::NotStarted);
 }
 
-FName FStateGraphNode::GetStateGraphName() const
-{
-	return StateGraphWeakPtr.IsValid() ? StateGraphWeakPtr.Pin()->GetName() : NAME_None;
-}
-
 void FStateGraphNode::UpdateConfig()
 {
 	GConfig->GetDouble(*ConfigSectionName, TEXT("Timeout"), Timeout, GEngineIni);
@@ -159,7 +155,7 @@ bool FStateGraphNodeFunction::CheckDependencies() const
 {
 	if (!StartFunction.IsBound())
 	{
-		UE_LOG_STATEGRAPH(Warning, TEXT("[%s.%s] Function node start not bound"), *GetStateGraphName().ToString(), *GetName().ToString());
+		UE_LOG_STATEGRAPH(Warning, TEXT("[%s] Function node start not bound"), *GetContextName());
 		return false;
 	}
 
@@ -171,8 +167,10 @@ void FStateGraphNodeFunction::Start()
 	// CheckDependencies verified the function is bound right before calling this in StateGraph->Run().
 	check(StartFunction.IsBound());
 
+	FStateGraph* StateGraph = GetStateGraph().Get();
+	check(StateGraph); // Must be valid since it is starting this node.
 	FStateGraphNodeWeakPtr WeakPtr(AsWeak());
-	StartFunction.Execute([WeakPtr, StateGraphName = GetStateGraphName(), NodeName = GetName()]()
+	StartFunction.Execute(*StateGraph, [WeakPtr, ContextName = GetContextName(), NodeName = GetName()]()
 	{
 		if (WeakPtr.IsValid())
 		{
@@ -180,23 +178,32 @@ void FStateGraphNodeFunction::Start()
 		}
 		else
 		{
-			UE_LOG_STATEGRAPH(Log, TEXT("[%s.%s] Function node completed after node was destroyed"), *StateGraphName.ToString(), *NodeName.ToString());
+			UE_LOG_STATEGRAPH(Log, TEXT("[%s] Function node completed after node was destroyed"), *ContextName);
 		}
 	});
 }
 
 // FStateGraph
 
-FStateGraph::FStateGraph(FName InName) :
+FStateGraph::FStateGraph(FName InName, const FString& InContextName) :
 	Name(InName),
 	ConfigSectionName(FString::Printf(TEXT("StateGraph.%s"), *InName.ToString()))
 {
-	UE_LOG_STATEGRAPH(VeryVerbose, TEXT("[%s] Created state graph"), *Name.ToString());
+	if (InContextName.IsEmpty())
+	{
+		ContextName = Name.ToString();
+	}
+	else
+	{
+		ContextName = FString::Printf(TEXT("%s(%s)"), *Name.ToString(), *InContextName);
+	}
+
+	UE_LOG_STATEGRAPH(VeryVerbose, TEXT("[%s] Created state graph"), *ContextName);
 }
 
 FStateGraph::~FStateGraph()
 {
-	UE_LOG_STATEGRAPH(VeryVerbose, TEXT("[%s] Destroyed state graph"), *Name.ToString());
+	UE_LOG_STATEGRAPH(VeryVerbose, TEXT("[%s] Destroyed state graph"), *ContextName);
 	FCoreDelegates::TSOnConfigSectionsChanged().Remove(ConfigSectionsChangedDelegate);
 }
 
@@ -244,17 +251,18 @@ bool FStateGraph::AddNode(const FStateGraphNodeRef& Node)
 {
 	if (Node->StateGraphWeakPtr.IsValid())
 	{
-		UE_LOG_STATEGRAPH(Warning, TEXT("[%s.%s] Node already associated with a state graph"), *Name.ToString(), *Node->GetName().ToString());
+		UE_LOG_STATEGRAPH(Warning, TEXT("[%s.%s] Node already associated with a state graph %s"), *ContextName, *Node->GetName().ToString(), *Node->ContextName);
 		return false;
 	}
 
 	if (Nodes.Contains(Node->GetName()))
 	{
-		UE_LOG_STATEGRAPH(Warning, TEXT("[%s.%s] Node with the same name already exists"), *Name.ToString(), *Node->GetName().ToString());
+		UE_LOG_STATEGRAPH(Warning, TEXT("[%s.%s] Node with the same name already exists"), *ContextName, *Node->GetName().ToString());
 		return false;
 	}
 
-	UE_LOG_STATEGRAPH(Verbose, TEXT("[%s.%s] Adding node"), *Name.ToString(), *Node->GetName().ToString());
+	Node->ContextName = FString::Printf(TEXT("%s.%s"), *ContextName, *Node->GetName().ToString());
+	UE_LOG_STATEGRAPH(Verbose, TEXT("[%s] Adding node"), *Node->ContextName);
 	Node->StateGraphWeakPtr = AsWeak();
 	Node->ConfigSectionName = FString::Printf(TEXT("%s.%s"), *ConfigSectionName, *Node->GetName().ToString());
 	Nodes.Add(Node->GetName(), Node);
@@ -267,11 +275,12 @@ bool FStateGraph::RemoveNode(FName NodeName)
 	FStateGraphNodeRef* Node = GetNodeRef(NodeName);
 	if (!Node)
 	{
-		UE_LOG_STATEGRAPH(Warning, TEXT("[%s.%s] Failed to remove node"), *Name.ToString(), *NodeName.ToString());
+		UE_LOG_STATEGRAPH(Warning, TEXT("[%s.%s] Failed to remove node"), *ContextName, *NodeName.ToString());
 		return false;
 	}
 
-	UE_LOG_STATEGRAPH(Verbose, TEXT("[%s.%s] Removing node"), *Name.ToString(), *NodeName.ToString());
+	UE_LOG_STATEGRAPH(Verbose, TEXT("[%s] Removing node"), *(*Node)->ContextName);
+	(*Node)->ContextName = (*Node)->GetName().ToString();
 	(*Node)->StateGraphWeakPtr.Reset();
 
 	// Keep a reference to check if the state graph is destroyed during external functions.
@@ -336,7 +345,7 @@ void FStateGraph::Run()
 		NextTimeout = (StartTime + Timeout) - Now;
 		if (NextTimeout <= 0.f)
 		{
-			UE_LOG_STATEGRAPH(Log, TEXT("[%s] State graph timed out (Duration=%.6f Timeout=%.6f)"), *Name.ToString(), Now - StartTime, Timeout);
+			UE_LOG_STATEGRAPH(Log, TEXT("[%s] State graph timed out (Duration=%.6f Timeout=%.6f)"), *ContextName, Now - StartTime, Timeout);
 			SetStatus(EStatus::TimedOut);
 			return;
 		}
@@ -344,7 +353,7 @@ void FStateGraph::Run()
 
 	bRunning = true;
 
-	UE_LOG_STATEGRAPH(Verbose, TEXT("[%s] Starting run loop (Now=%.06f)"), *Name.ToString(), Now);
+	UE_LOG_STATEGRAPH(Verbose, TEXT("[%s] Starting run loop (Now=%.06f)"), *ContextName, Now);
 
 	uint32 Blocked = 0;
 	uint32 Started = 0;
@@ -405,7 +414,7 @@ void FStateGraph::Run()
 				break;
 			}
 
-			UE_LOG_STATEGRAPH(Log, TEXT("[%s.%s] Starting node"), *Name.ToString(), *NodeName.ToString());
+			UE_LOG_STATEGRAPH(Log, TEXT("[%s] Starting node"), *(*Node)->ContextName);
 			(*Node)->SetStatus(FStateGraphNode::EStatus::Started);
 
 			if (!StateGraphWeakPtr.IsValid())
@@ -462,7 +471,7 @@ void FStateGraph::Run()
 				if (NodeTimeout <= 0.f)
 				{
 					++TimedOut;
-					UE_LOG_STATEGRAPH(Log, TEXT("[%s.%s] Node timed out (Duration=%.6f Timeout=%.6f)"), *Name.ToString(), *NodeName.ToString(), Now - (*Node)->StartTime, (*Node)->Timeout);
+					UE_LOG_STATEGRAPH(Log, TEXT("[%s] Node timed out (Duration=%.6f Timeout=%.6f)"), *(*Node)->ContextName, Now - (*Node)->StartTime, (*Node)->Timeout);
 					(*Node)->SetStatus(FStateGraphNode::EStatus::TimedOut);
 
 					if (!StateGraphWeakPtr.IsValid())
@@ -505,7 +514,7 @@ void FStateGraph::Run()
 			break;
 
 		default:
-			UE_LOG_STATEGRAPH(Error, TEXT("[%s.%s] Unknown state"), *Name.ToString(), *NodeName.ToString());
+			UE_LOG_STATEGRAPH(Error, TEXT("[%s] Unknown state"), *(*Node)->ContextName);
 			checkNoEntry();
 			break;
 		}
@@ -513,7 +522,7 @@ void FStateGraph::Run()
 
 	bRunning = false;
 	UE_LOG_STATEGRAPH(Verbose, TEXT("[%s] Duration=%.6f Timeout=%.6f Blocked=%d Started=%d Running=%d Completed=%d Removed=%d TimedOut=%d"),
-		*Name.ToString(), Now - StartTime, Timeout, Blocked, Started, Running, Completed, Removed, TimedOut);
+		*ContextName, Now - StartTime, Timeout, Blocked, Started, Running, Completed, Removed, TimedOut);
 
 	if (bRunAgain)
 	{
@@ -532,19 +541,19 @@ void FStateGraph::Run()
 		if (Blocked == 0 && TimedOut == 0)
 		{
 			CompletedTime = Now;
-			UE_LOG_STATEGRAPH(Log, TEXT("[%s] Completed (Duration=%.6f Timeout=%.6f)"), *Name.ToString(), CompletedTime - StartTime, Timeout);
+			UE_LOG_STATEGRAPH(Log, TEXT("[%s] Completed (Duration=%.6f Timeout=%.6f)"), *ContextName, CompletedTime - StartTime, Timeout);
 			SetStatus(EStatus::Completed);
 		}
 		else
 		{
-			UE_LOG_STATEGRAPH(Warning, TEXT("[%s] Blocked on %d nodes, timed out %d nodes"), *Name.ToString(), Blocked, TimedOut);
+			UE_LOG_STATEGRAPH(Warning, TEXT("[%s] Blocked on %d nodes, timed out %d nodes"), *ContextName, Blocked, TimedOut);
 			LogDebugInfo(true);
 			SetStatus(EStatus::Blocked);
 		}
 	}
 	else
 	{
-		UE_LOG_STATEGRAPH(VeryVerbose, TEXT("[%s] Waiting on %d nodes"), *Name.ToString(), Started + Running);
+		UE_LOG_STATEGRAPH(VeryVerbose, TEXT("[%s] Waiting on %d nodes"), *ContextName, Started + Running);
 		SetStatus(EStatus::Waiting);
 	}
 
@@ -555,7 +564,7 @@ void FStateGraph::Run()
 
 	if ((Status == EStatus::Blocked || Status == EStatus::Waiting) && NextTimeout > 0.f)
 	{
-		UE_LOG_STATEGRAPH(Verbose, TEXT("[%s] Setting timer for %.6f"), *Name.ToString(), NextTimeout);
+		UE_LOG_STATEGRAPH(Verbose, TEXT("[%s] Setting timer for %.6f"), *ContextName, NextTimeout);
 		TimeoutTicker = FTSTicker::GetCoreTicker().AddTicker(FTickerDelegate::CreateSPLambda(this, [this](float DeltaTime) {
 			Run();
 			return false;
@@ -570,7 +579,7 @@ void FStateGraph::Reset()
 	// Keep a reference to check if the state graph is destroyed during external functions.
 	FStateGraphWeakPtr StateGraphWeakPtr(AsWeak());
 
-	UE_LOG_STATEGRAPH(Verbose, TEXT("[%s] Resetting state graph"), *Name.ToString());
+	UE_LOG_STATEGRAPH(Verbose, TEXT("[%s] Resetting state graph"), *ContextName);
 	StartTime = 0.f;
 	CompletedTime = 0.f;
 	bRunAgain = false;
@@ -599,7 +608,7 @@ void FStateGraph::Reset()
 
 void FStateGraph::Pause()
 {
-	UE_LOG_STATEGRAPH(Verbose, TEXT("[%s] Pausing state graph"), *Name.ToString());
+	UE_LOG_STATEGRAPH(Verbose, TEXT("[%s] Pausing state graph"), *ContextName);
 	bRunAgain = false;
 	SetStatus(EStatus::Paused);
 }
@@ -611,7 +620,7 @@ void FStateGraph::LogDebugInfo(bool bWarning)
 		return;
 	}
 
-	FString Message = FString::Printf(TEXT("[%s] Status=%s Nodes=%d Duration=%.6f Timeout=%.6f"), *Name.ToString(), GetStatusName(), Nodes.Num(), GetDuration(), Timeout);
+	FString Message = FString::Printf(TEXT("[%s] Status=%s Nodes=%d Duration=%.6f Timeout=%.6f"), *ContextName, GetStatusName(), Nodes.Num(), GetDuration(), Timeout);
 
 	if (bWarning)
 	{
@@ -657,7 +666,7 @@ void FStateGraph::LogDebugInfo(bool bWarning)
 		}
 
 		Message = FString::Printf(TEXT("[%s.%s] Status=%s Duration=%.6f Timeout=%.6f Dependencies(%s)"),
-			*Name.ToString(), *Node.Key.ToString(), Node.Value->GetStatusName(), Node.Value->GetDuration(), Node.Value->Timeout, *FString::Join(Dependencies, TEXT(" ")));
+			*ContextName, *Node.Key.ToString(), Node.Value->GetStatusName(), Node.Value->GetDuration(), Node.Value->Timeout, *FString::Join(Dependencies, TEXT(" ")));
 
 		if (bWarning)
 		{
diff --git a/Engine/Plugins/Experimental/StateGraph/Source/StateGraph/Public/StateGraph.h b/Engine/Plugins/Experimental/StateGraph/Source/StateGraph/Public/StateGraph.h
--- a/Engine/Plugins/Experimental/StateGraph/Source/StateGraph/Public/StateGraph.h
+++ b/Engine/Plugins/Experimental/StateGraph/Source/StateGraph/Public/StateGraph.h
@@ -48,6 +48,11 @@ public:
 		return Name;
 	}
 
+	const FString& GetContextName() const
+	{
+		return ContextName;
+	}
+
 	enum class EStatus : uint8
 	{
 		/** State graph has not attempted to start this node yet, or the node has been reset. */
@@ -103,9 +108,6 @@ public:
 		return StateGraphWeakPtr.Pin();
 	}
 
-	/** Get the name of the state graph, or NAME_None if the state graph is invalid. */
-	FName GetStateGraphName() const;
-
 	/**
 	 * Helper function to create a new node, add it to this node's StateGraph, and add this node as a dependency.
 	 * This allows for chaining nodes together such as:
@@ -113,7 +115,7 @@ public:
 	 *	StateGraph->CreateNode("NodeA", this, &SomeMemberFunction)
 	 *		->Next("NodeB", &SomeStaticFunction)
 	 *		->Next<SomeNodeClass>()
-	 *		->Next("NodeC", [](FStateGraphNodeFunctionComplete Complete) { Complete(); return true; });
+	 *		->Next("NodeC", [](FStateGraph& StateGraph, FStateGraphNodeFunctionComplete Complete) { Complete(); });
 	 */
 	template<typename NodeType, typename... ArgsTypes>
 	TSharedPtr<NodeType, ESPMode::ThreadSafe> Next(ArgsTypes&&... Args) const;
@@ -149,6 +151,7 @@ private:
 	void SetStatus(EStatus NewStatus);
 
 	FName Name;
+	FString ContextName;
 	FString ConfigSectionName;
 	EStatus Status = EStatus::NotStarted;
 	FStateGraphWeakPtr StateGraphWeakPtr;
@@ -160,7 +163,7 @@ private:
 };
 
 /** Start function delegate type for FStateGraphNodeFunction nodes. */
-DECLARE_DELEGATE_OneParam(FStateGraphNodeFunctionStart, FStateGraphNodeFunctionComplete /*Complete*/);
+DECLARE_DELEGATE_TwoParams(FStateGraphNodeFunctionStart, FStateGraph& /*StateGraph*/, FStateGraphNodeFunctionComplete /*Complete*/);
 
 /**
  * Node for wrapping functions as delegates.
@@ -187,7 +190,7 @@ class STATEGRAPH_API FStateGraph : public TSharedFromThis<FStateGraph>, public F
 {
 public:
 
-	FStateGraph(FName InName);
+	FStateGraph(FName InName, const FString& InContextName = FString());
 
 	virtual ~FStateGraph();
 
@@ -200,6 +203,11 @@ public:
 		return Name;
 	}
 
+	const FString& GetContextName() const
+	{
+		return ContextName;
+	}
+
 	enum class EStatus : uint8
 	{
 		/** State graph has not started yet. */
@@ -291,7 +299,7 @@ public:
 
 	/** Create a new function node for member functions. */
 	template<typename ObjectType, typename FunctionType, typename... ArgsTypes>
-	typename TEnableIf<TIsClass<ObjectType>::Value && TIsInvocable<FunctionType, ObjectType, FStateGraphNodeFunctionComplete, ArgsTypes...>::Value, FStateGraphNodeFunctionRef>::Type
+	typename TEnableIf<TIsClass<ObjectType>::Value && TIsInvocable<FunctionType, ObjectType, FStateGraph&, FStateGraphNodeFunctionComplete, ArgsTypes...>::Value, FStateGraphNodeFunctionRef>::Type
 		CreateNode(FName NodeName, ObjectType* Object, FunctionType Function, ArgsTypes&&... Args)
 	{
 		if constexpr (IsDerivedFromSharedFromThis<ObjectType>())
@@ -310,7 +318,7 @@ public:
 
 	/** Create a new function node for lambdas tied to objects. */
 	template<typename ObjectType, typename FunctionType, typename... ArgsTypes>
-	typename TEnableIf<TIsClass<ObjectType>::Value && TIsInvocable<FunctionType, FStateGraphNodeFunctionComplete, ArgsTypes...>::Value, FStateGraphNodeFunctionRef>::Type
+	typename TEnableIf<TIsClass<ObjectType>::Value && TIsInvocable<FunctionType, FStateGraph&, FStateGraphNodeFunctionComplete, ArgsTypes...>::Value, FStateGraphNodeFunctionRef>::Type
 		CreateNode(FName NodeName, ObjectType* Object, FunctionType Function, ArgsTypes&&... Args)
 	{
 		if constexpr (IsDerivedFromSharedFromThis<ObjectType>())
@@ -325,7 +333,7 @@ public:
 
 	/** Create a new function node for lambas and static methods. */
 	template<typename FunctionType, typename... ArgsTypes>
-	typename TEnableIf<TIsInvocable<FunctionType, FStateGraphNodeFunctionComplete, ArgsTypes...>::Value, FStateGraphNodeFunctionRef>::Type
+	typename TEnableIf<TIsInvocable<FunctionType, FStateGraph&, FStateGraphNodeFunctionComplete, ArgsTypes...>::Value, FStateGraphNodeFunctionRef>::Type
 		CreateNode(FName NodeName, FunctionType Function, ArgsTypes&&... Args)
 	{
 		return CreateNode<FStateGraphNodeFunction>(NodeName, FStateGraphNodeFunctionStart::CreateLambda(Function, Forward<ArgsTypes>(Args)...));
@@ -365,6 +373,7 @@ protected:
 	void OnConfigSectionsChanged(const FString& IniFilename, const TSet<FString>& SectionNames);
 
 	FName Name;
+	FString ContextName;
 	FString ConfigSectionName;
 	EStatus Status = EStatus::NotStarted;
 	double StartTime = 0.f;
@@ -384,7 +393,7 @@ TSharedPtr<NodeType, ESPMode::ThreadSafe> FStateGraphNode::Next(ArgsTypes&&... A
 {
 	if (!StateGraphWeakPtr.IsValid())
 	{
-		UE_LOG_STATEGRAPH(Warning, TEXT("[%s] Node attempted to be created with invalid state graph"), *Name.ToString());
+		UE_LOG_STATEGRAPH(Warning, TEXT("[%s] Node attempted to be created with invalid state graph"), *ContextName);
 		return TSharedPtr<NodeType, ESPMode::ThreadSafe>();
 	}
 
@@ -398,7 +407,7 @@ FStateGraphNodeFunctionPtr FStateGraphNode::Next(ArgsTypes&&... Args) const
 {
 	if (!StateGraphWeakPtr.IsValid())
 	{
-		UE_LOG_STATEGRAPH(Warning, TEXT("[%s] Node attempted to be created with invalid state graph"), *Name.ToString());
+		UE_LOG_STATEGRAPH(Warning, TEXT("[%s] Node attempted to be created with invalid state graph"), *ContextName);
 		return FStateGraphNodeFunctionPtr();
 	}
 
diff --git a/Engine/Plugins/Experimental/StateGraph/Tests/StateGraphTests/Private/StateGraphTestsExample.cpp b/Engine/Plugins/Experimental/StateGraph/Tests/StateGraphTests/Private/StateGraphTestsExample.cpp
--- a/Engine/Plugins/Experimental/StateGraph/Tests/StateGraphTests/Private/StateGraphTestsExample.cpp
+++ b/Engine/Plugins/Experimental/StateGraph/Tests/StateGraphTests/Private/StateGraphTestsExample.cpp
@@ -107,8 +107,8 @@ public:
 		// CreateNode supports any function type that can be used to create delegates with: lambdas, members, static, raw, etc.
 
 		// Since NodeA and NodeB have no dependencies, they run first concurrently.
-		StateGraph->CreateNode("Base.NodeA", this, [](FStateGraphNodeFunctionComplete Complete) {	Complete(); });
-		StateGraph->CreateNode("Base.NodeB", this, [](FStateGraphNodeFunctionComplete Complete) {	ExampleAsync(0.1f, Complete); });
+		StateGraph->CreateNode("Base.NodeA", this, [](FStateGraph& StateGraph, FStateGraphNodeFunctionComplete Complete) { Complete(); });
+		StateGraph->CreateNode("Base.NodeB", this, [](FStateGraph& StateGraph, FStateGraphNodeFunctionComplete Complete) { ExampleAsync(0.1f, Complete); });
 
 		// CreateNode also supports custom node types by explicitly specifying the type and passing arguments to the constructor.
 		FExampleClassNodePtr NodeC = StateGraph->CreateNode<FExampleClassNode>("Base.NodeC");
@@ -125,12 +125,12 @@ public:
 
 protected:
 
-	static void NodeD(FStateGraphNodeFunctionComplete Complete)
+	static void NodeD(FStateGraph& InStateGraph, FStateGraphNodeFunctionComplete Complete)
 	{
 		Complete();
 	}
 
-	void NodeE(FStateGraphNodeFunctionComplete Complete)
+	void NodeE(FStateGraph& InStateGraph, FStateGraphNodeFunctionComplete Complete)
 	{
 		FExampleClassNodePtr NodeC = StateGraph->GetNode<FExampleClassNode>("Base.NodeC");
 		check(NodeC.IsValid());
@@ -154,14 +154,14 @@ public:
 	void Run(const FExampleModuleComplete& ModuleComplete)
 	{
 		// Run nodes in sequence to produce some output that the caller needs, using the same state graph which may be running unrelated nodes.
-		StateGraph->CreateNode("Module.NodeA", this, [](FStateGraphNodeFunctionComplete Complete) { Complete(); })
+		StateGraph->CreateNode("Module.NodeA", this, [](FStateGraph& StateGraph, FStateGraphNodeFunctionComplete Complete) { Complete(); })
 			->Next("Module.NodeB", this, &FExampleModule::NodeB)
-			->Next("Module.NodeC", this, [ModuleComplete](FStateGraphNodeFunctionComplete Complete) { Complete(); ModuleComplete.ExecuteIfBound(42); });
+			->Next("Module.NodeC", this, [ModuleComplete](FStateGraph& StateGraph, FStateGraphNodeFunctionComplete Complete) { Complete(); ModuleComplete.ExecuteIfBound(42); });
 	}
 
 protected:
 
-	void NodeB(FStateGraphNodeFunctionComplete Complete)
+	void NodeB(FStateGraph& InStateGraph, FStateGraphNodeFunctionComplete Complete)
 	{
 		ExampleAsync(0.1f, Complete);
 	}
@@ -194,13 +194,13 @@ public:
 		(*NodeDRef)->Dependencies.Add("Derived.NodeF");
 
 		// Add some additional nodes.
-		StateGraph->CreateNode("Derived.NodeG", this, [](FStateGraphNodeFunctionComplete Complete) { Complete(); })
+		StateGraph->CreateNode("Derived.NodeG", this, [](FStateGraph& InStateGraph, FStateGraphNodeFunctionComplete Complete) { Complete(); })
 			->Next<FExampleClassNode>("Derived.NodeH");
 	}
 
 protected:
 
-	void NodeF(FStateGraphNodeFunctionComplete Complete)
+	void NodeF(FStateGraph& InStateGraph, FStateGraphNodeFunctionComplete Complete)
 	{
 		FExampleModulePtr ModulePtr = MakeShared<FExampleModule>(StateGraph);
 		ModulePtr->Run(FExampleModuleComplete::CreateSPLambda(this, [ModulePtr, Complete](int32 Output)
diff --git a/Engine/Plugins/Experimental/StateGraph/Tests/StateGraphTests/Private/StateGraphTestsUnit.cpp b/Engine/Plugins/Experimental/StateGraph/Tests/StateGraphTests/Private/StateGraphTestsUnit.cpp
--- a/Engine/Plugins/Experimental/StateGraph/Tests/StateGraphTests/Private/StateGraphTestsUnit.cpp
+++ b/Engine/Plugins/Experimental/StateGraph/Tests/StateGraphTests/Private/StateGraphTestsUnit.cpp
@@ -72,7 +72,7 @@ TEST_CASE("FStateGraphNodeFunction Basic Tests", "[FStateGraphNodeFunction]")
 {
 	FStateGraphRef StateGraph(MakeShared<FStateGraph>("Test"));
 	FStateGraphNodeFunctionComplete Complete;
-	FStateGraphNodeFunctionRef TestNodeA = StateGraph->CreateNode("TestNodeA", [&Complete](FStateGraphNodeFunctionComplete InComplete) { Complete = InComplete; });
+	FStateGraphNodeFunctionRef TestNodeA = StateGraph->CreateNode("TestNodeA", [&Complete](FStateGraph& StateGraph, FStateGraphNodeFunctionComplete InComplete) { Complete = InComplete; });
 	CHECK(TestNodeA->GetStatus() == FStateGraphNode::EStatus::NotStarted);
 	CHECK(FString(TestNodeA->GetStatusName()) == FString(TEXT("NotStarted")));
 
@@ -85,6 +85,15 @@ TEST_CASE("FStateGraphNodeFunction Basic Tests", "[FStateGraphNodeFunction]")
 	CHECK(TestNodeA->GetStatus() == FStateGraphNode::EStatus::Completed);
 }
 
+TEST_CASE("FStateGraph ContextName", "[FStateGraphContextName]")
+{
+	FStateGraphRef StateGraph(MakeShared<FStateGraph>("Test", TEXT("TestContextName")));
+	FStateGraphNodeFunctionRef TestNodeA = StateGraph->CreateNode("TestNodeA", [](FStateGraph& StateGraph, FStateGraphNodeFunctionComplete Complete) { Complete(); });
+	StateGraph->Run();
+	CHECK(StateGraph->GetContextName() == TEXT("Test(TestContextName)"));
+	CHECK(TestNodeA->GetContextName() == TEXT("Test(TestContextName).TestNodeA"));
+}
+
 class FClassNode : public FStateGraphNode
 {
 public:
@@ -92,18 +101,18 @@ public:
 	virtual void Start() { Complete(); }
 };
 
-void StaticNode(FStateGraphNodeFunctionComplete Complete) { Complete(); }
+void StaticNode(FStateGraph& StateGraph, FStateGraphNodeFunctionComplete Complete) { Complete(); }
 
 class FRawClass
 {
 public:
-	void RawNode(FStateGraphNodeFunctionComplete Complete) { Complete(); }
+	void RawNode(FStateGraph& StateGraph, FStateGraphNodeFunctionComplete Complete) { Complete(); }
 };
 
 class FSPClass : public TSharedFromThis<FSPClass>
 {
 public:
-	void SPNode(FStateGraphNodeFunctionComplete Complete) { Complete(); }
+	void SPNode(FStateGraph& StateGraph, FStateGraphNodeFunctionComplete Complete) { Complete(); }
 };
 
 TEST_CASE("CreateNode for each supported type", "[CreateNodeFunctions]")
@@ -113,11 +122,11 @@ TEST_CASE("CreateNode for each supported type", "[CreateNodeFunctions]")
 	TSharedRef<FSPClass> SPClass(MakeShared<FSPClass>());
 
 	StateGraph->CreateNode<FClassNode>("TestClass")
-		->Next("TestLambda", [](FStateGraphNodeFunctionComplete Complete) { Complete(); })
+		->Next("TestLambda", [](FStateGraph& StateGraph, FStateGraphNodeFunctionComplete Complete) { Complete(); })
 		->Next("TestStatic", &StaticNode)
 		->Next("TestRaw", &RawClass, &FRawClass::RawNode)
 		->Next("TestSP", &SPClass.Get(), &FSPClass::SPNode)
-		->Next("TestSPLambda", &SPClass.Get(), [](FStateGraphNodeFunctionComplete Complete) { Complete(); });
+		->Next("TestSPLambda", &SPClass.Get(), [](FStateGraph& StateGraph, FStateGraphNodeFunctionComplete Complete) { Complete(); });
 
 	StateGraph->Run();
 	CHECK(StateGraph->GetStatus() == FStateGraph::EStatus::Completed);
@@ -233,10 +242,10 @@ TEST_CASE("Removing nodes", "[RemoveNodes]")
 	CHECK(!WeakTestNodeA.IsValid());
 
 	WeakTestNodeA = StateGraph->CreateNode<FTestNode>("TestNodeA");
-	FStateGraphNodePtr TestNodeB = StateGraph->CreateNode("TestNodeB", [StateGraph](FStateGraphNodeFunctionComplete Complete)
+	FStateGraphNodePtr TestNodeB = StateGraph->CreateNode("TestNodeB", [](FStateGraph& StateGraph, FStateGraphNodeFunctionComplete Complete)
 	{
-		StateGraph->RemoveNode("TestNodeA");
-		StateGraph->RemoveNode("TestNodeC");
+		StateGraph.RemoveNode("TestNodeA");
+		StateGraph.RemoveNode("TestNodeC");
 	});
 	FStateGraphNodeWeakPtr WeakTestNodeC = StateGraph->CreateNode<FTestNode>("TestNodeC");
 	StateGraph->Run();
@@ -258,7 +267,7 @@ TEST_CASE("Resetting nodes", "[ResetNodes]")
 	FStateGraphRef StateGraph(MakeShared<FStateGraph>("Test"));
 
 	int32 StartCount = 0;
-	FStateGraphNodeRef TestNodeA = StateGraph->CreateNode("TestNodeA", [&StartCount, &TestNodeA](FStateGraphNodeFunctionComplete Complete)
+	FStateGraphNodeRef TestNodeA = StateGraph->CreateNode("TestNodeA", [&StartCount, &TestNodeA](FStateGraph& StateGraph, FStateGraphNodeFunctionComplete Complete)
 	{
 		CHECK(TestNodeA->GetStatus() == FStateGraphNode::EStatus::Started);
 
@@ -292,11 +301,11 @@ TEST_CASE("Resetting state graph", "[ResetStateGraph]")
 	FStateGraphRef StateGraph(MakeShared<FStateGraph>("Test"));
 
 	int32 StartCount = 0;
-	FStateGraphNodeRef TestNodeA = StateGraph->CreateNode("TestNodeA", [&StartCount, &StateGraph](FStateGraphNodeFunctionComplete Complete)
+	FStateGraphNodeRef TestNodeA = StateGraph->CreateNode("TestNodeA", [&StartCount](FStateGraph& StateGraph, FStateGraphNodeFunctionComplete Complete)
 	{
 		if (StartCount == 0)
 		{
-			StateGraph->Reset();
+			StateGraph.Reset();
 		}
 		else
 		{
@@ -321,7 +330,7 @@ TEST_CASE("Deleting state graph", "[DeleteStateGraph]")
 {
 	FStateGraphPtr StateGraph(MakeShared<FStateGraph>("Test"));
 	FStateGraphWeakPtr WeakStateGraph(StateGraph);
-	FStateGraphNodeRef TestNodeA = StateGraph->CreateNode("TestNodeA", [&StateGraph, &WeakStateGraph](FStateGraphNodeFunctionComplete Complete)
+	FStateGraphNodeRef TestNodeA = StateGraph->CreateNode("TestNodeA", [&StateGraph, &WeakStateGraph](FStateGraph& InStateGraph, FStateGraphNodeFunctionComplete Complete)
 	{
 		StateGraph.Reset();
 		CHECK(!WeakStateGraph.IsValid());
@@ -329,15 +338,14 @@ TEST_CASE("Deleting state graph", "[DeleteStateGraph]")
 
 	StateGraph->Run();
 	CHECK(!WeakStateGraph.IsValid());
-	CHECK(TestNodeA->GetStateGraphName() == NAME_None);
 }
 
 TEST_CASE("Pausing state graph", "[PauseStateGraph]")
 {
 	FStateGraphRef StateGraph(MakeShared<FStateGraph>("Test"));
-	FStateGraphNodeRef TestNodeA = StateGraph->CreateNode("TestNodeA", [&StateGraph](FStateGraphNodeFunctionComplete Complete) { StateGraph->Pause(); });
+	FStateGraphNodeRef TestNodeA = StateGraph->CreateNode("TestNodeA", [](FStateGraph& StateGraph, FStateGraphNodeFunctionComplete Complete) { StateGraph.Pause(); });
 	FStateGraphNodeRef TestNodeB = StateGraph->CreateNode<FTestNode>("TestNodeB", TSet<FName>({"TestNodeA"}));
-	FStateGraphNodeRef TestNodeC = StateGraph->CreateNode("TestNodeC", [](FStateGraphNodeFunctionComplete Complete) { Complete(); });
+	FStateGraphNodeRef TestNodeC = StateGraph->CreateNode("TestNodeC", [](FStateGraph& StateGraph, FStateGraphNodeFunctionComplete Complete) { Complete(); });
 	TestNodeC->Dependencies.Add("TestNodeB");
 
 	StateGraph->Run();
@@ -379,7 +387,7 @@ TEST_CASE("State graph timeout", "[StateGraphTimeout]")
 		}
 	});
 
-	FStateGraphNodeRef TestNodeA = StateGraph->CreateNode("TestNodeA", [](FStateGraphNodeFunctionComplete Complete) {});
+	FStateGraphNodeRef TestNodeA = StateGraph->CreateNode("TestNodeA", [](FStateGraph& StateGraph, FStateGraphNodeFunctionComplete Complete) {});
 	StateGraph->Run();
 
 	while (StateGraph->GetStatus() == FStateGraph::EStatus::Waiting)
@@ -430,7 +438,7 @@ TEST_CASE("State graph config", "[StateGraphConfig]")
 
 	FStateGraphRef StateGraph(MakeShared<FStateGraph>("Test"));
 	StateGraph->Initialize();
-	FStateGraphNodeRef TestNodeA = StateGraph->CreateNode("TestNodeA", [](FStateGraphNodeFunctionComplete Complete) {});
+	FStateGraphNodeRef TestNodeA = StateGraph->CreateNode("TestNodeA", [](FStateGraph& StateGraph, FStateGraphNodeFunctionComplete Complete) {});
 	StateGraph->Run();
 
 	while (StateGraph->GetStatus() == FStateGraph::EStatus::Waiting)
