StateGraph Updates:
- Stop using __FUNCTION__ in logging.
- Update LogDebugInfo to take verbosity for usage outside the module.
- Change WeakPtr IsValid/Pin logic to prevent race conditions.
#rnx

diff --git a/Engine/Plugins/Experimental/StateGraph/Source/StateGraph/Private/StateGraph.cpp b/Engine/Plugins/Experimental/StateGraph/Source/StateGraph/Private/StateGraph.cpp
--- a/Engine/Plugins/Experimental/StateGraph/Source/StateGraph/Private/StateGraph.cpp
+++ b/Engine/Plugins/Experimental/StateGraph/Source/StateGraph/Private/StateGraph.cpp
@@ -24,12 +24,12 @@ FStateGraphNode::FStateGraphNode(FName InName) :
 	Name(InName),
 	ContextName(Name.ToString())
 {
-	UE_LOG_STATEGRAPH(VeryVerbose, TEXT("[%s] Created node"), *ContextName);
+	UE_LOG(LogStateGraph, VeryVerbose, TEXT("[%s] Created node"), *ContextName);
 }
 
 FStateGraphNode::~FStateGraphNode()
 {
-	UE_LOG_STATEGRAPH(VeryVerbose, TEXT("[%s] Destroyed node"), *ContextName);
+	UE_LOG(LogStateGraph, VeryVerbose, TEXT("[%s] Destroyed node"), *ContextName);
 }
 
 const TCHAR* FStateGraphNode::GetStatusName(EStatus Status)
@@ -65,7 +65,7 @@ bool FStateGraphNode::CheckDependencies() const
 	FStateGraphPtr StateGraphPtr(StateGraphWeakPtr.Pin());
 	if (!StateGraphPtr.IsValid())
 	{
-		UE_LOG_STATEGRAPH(Warning, TEXT("[%s] Node checked with invalid state graph"), *ContextName);
+		UE_LOG(LogStateGraph, Warning, TEXT("[%s] Node checked with invalid state graph"), *ContextName);
 		return false;
 	}
 
@@ -85,12 +85,12 @@ void FStateGraphNode::Complete()
 {
 	if (Status == EStatus::Completed)
 	{
-		UE_LOG_STATEGRAPH(Warning, TEXT("[%s] Node already completed"), *ContextName);
+		UE_LOG(LogStateGraph, Warning, TEXT("[%s] Node already completed"), *ContextName);
 		return;
 	}
 
 	CompletedTime = FPlatformTime::Seconds();
-	UE_LOG_STATEGRAPH(Log, TEXT("[%s] Completed node (Duration=%.6f Timeout=%.6f)"), *ContextName, CompletedTime - StartTime, Timeout);
+	UE_LOG(LogStateGraph, Log, TEXT("[%s] Completed node (Duration=%.6f Timeout=%.6f)"), *ContextName, CompletedTime - StartTime, Timeout);
 
 	// Keep a reference to check if the node is destroyed during external functions.
 	FStateGraphNodeWeakPtr StateGraphNodeWeakPtr(AsWeak());
@@ -113,13 +113,13 @@ void FStateGraphNode::Complete()
 	}
 	else
 	{
-		UE_LOG_STATEGRAPH(Warning, TEXT("[%s] Node completed with invalid state graph"), *ContextName);
+		UE_LOG(LogStateGraph, Warning, TEXT("[%s] Node completed with invalid state graph"), *ContextName);
 	}
 }
 
 void FStateGraphNode::Reset()
 {
-	UE_LOG_STATEGRAPH(Verbose, TEXT("[%s] Resetting node"), *ContextName);
+	UE_LOG(LogStateGraph, Verbose, TEXT("[%s] Resetting node"), *ContextName);
 	StartTime = 0.f;
 	CompletedTime = 0.f;
 	SetStatus(EStatus::NotStarted);
@@ -155,7 +155,7 @@ bool FStateGraphNodeFunction::CheckDependencies() const
 {
 	if (!StartFunction.IsBound())
 	{
-		UE_LOG_STATEGRAPH(Warning, TEXT("[%s] Function node start not bound"), *GetContextName());
+		UE_LOG(LogStateGraph, Warning, TEXT("[%s] Function node start not bound"), *GetContextName());
 		return false;
 	}
 
@@ -169,18 +169,20 @@ void FStateGraphNodeFunction::Start()
 
 	FStateGraph* StateGraph = GetStateGraph().Get();
 	check(StateGraph); // Must be valid since it is starting this node.
-	FStateGraphNodeWeakPtr WeakPtr(AsWeak());
-	StartFunction.Execute(*StateGraph, [WeakPtr, ContextName = GetContextName(), NodeName = GetName()]()
-	{
-		if (WeakPtr.IsValid())
+	FStateGraphNodeWeakPtr NodeWeakPtr(AsWeak());
+	StartFunction.Execute(*StateGraph,
+		[NodeWeakPtr, ContextName = GetContextName()]()
 		{
-			WeakPtr.Pin()->Complete();
-		}
-		else
-		{
-			UE_LOG_STATEGRAPH(Log, TEXT("[%s] Function node completed after node was destroyed"), *ContextName);
-		}
-	});
+			FStateGraphNodePtr NodePtr = NodeWeakPtr.Pin();
+			if (NodePtr)
+			{
+				NodePtr->Complete();
+			}
+			else
+			{
+				UE_LOG(LogStateGraph, Log, TEXT("[%s] Function node completed after node was destroyed"), *ContextName);
+			}
+		});
 }
 
 // FStateGraph
@@ -198,12 +200,12 @@ FStateGraph::FStateGraph(FName InName, const FString& InContextName) :
 		ContextName = FString::Printf(TEXT("%s(%s)"), *Name.ToString(), *InContextName);
 	}
 
-	UE_LOG_STATEGRAPH(VeryVerbose, TEXT("[%s] Created state graph"), *ContextName);
+	UE_LOG(LogStateGraph, VeryVerbose, TEXT("[%s] Created state graph"), *ContextName);
 }
 
 FStateGraph::~FStateGraph()
 {
-	UE_LOG_STATEGRAPH(VeryVerbose, TEXT("[%s] Destroyed state graph"), *ContextName);
+	UE_LOG(LogStateGraph, VeryVerbose, TEXT("[%s] Destroyed state graph"), *ContextName);
 	FCoreDelegates::TSOnConfigSectionsChanged().Remove(ConfigSectionsChangedDelegate);
 }
 
@@ -251,18 +253,18 @@ bool FStateGraph::AddNode(const FStateGraphNodeRef& Node)
 {
 	if (Node->StateGraphWeakPtr.IsValid())
 	{
-		UE_LOG_STATEGRAPH(Warning, TEXT("[%s.%s] Node already associated with a state graph %s"), *ContextName, *Node->GetName().ToString(), *Node->ContextName);
+		UE_LOG(LogStateGraph, Warning, TEXT("[%s.%s] Node already associated with a state graph %s"), *ContextName, *Node->GetName().ToString(), *Node->ContextName);
 		return false;
 	}
 
 	if (Nodes.Contains(Node->GetName()))
 	{
-		UE_LOG_STATEGRAPH(Warning, TEXT("[%s.%s] Node with the same name already exists"), *ContextName, *Node->GetName().ToString());
+		UE_LOG(LogStateGraph, Warning, TEXT("[%s.%s] Node with the same name already exists"), *ContextName, *Node->GetName().ToString());
 		return false;
 	}
 
 	Node->ContextName = FString::Printf(TEXT("%s.%s"), *ContextName, *Node->GetName().ToString());
-	UE_LOG_STATEGRAPH(Verbose, TEXT("[%s] Adding node"), *Node->ContextName);
+	UE_LOG(LogStateGraph, Verbose, TEXT("[%s] Adding node"), *Node->ContextName);
 	Node->StateGraphWeakPtr = AsWeak();
 	Node->ConfigSectionName = FString::Printf(TEXT("%s.%s"), *ConfigSectionName, *Node->GetName().ToString());
 	Nodes.Add(Node->GetName(), Node);
@@ -275,11 +277,11 @@ bool FStateGraph::RemoveNode(FName NodeName)
 	FStateGraphNodeRef* Node = GetNodeRef(NodeName);
 	if (!Node)
 	{
-		UE_LOG_STATEGRAPH(Warning, TEXT("[%s.%s] Failed to remove node"), *ContextName, *NodeName.ToString());
+		UE_LOG(LogStateGraph, Warning, TEXT("[%s.%s] Failed to remove node"), *ContextName, *NodeName.ToString());
 		return false;
 	}
 
-	UE_LOG_STATEGRAPH(Verbose, TEXT("[%s] Removing node"), *(*Node)->ContextName);
+	UE_LOG(LogStateGraph, Verbose, TEXT("[%s] Removing node"), *(*Node)->ContextName);
 	(*Node)->ContextName = (*Node)->GetName().ToString();
 	(*Node)->StateGraphWeakPtr.Reset();
 
@@ -345,7 +347,7 @@ void FStateGraph::Run()
 		NextTimeout = (StartTime + Timeout) - Now;
 		if (NextTimeout <= 0.f)
 		{
-			UE_LOG_STATEGRAPH(Log, TEXT("[%s] State graph timed out (Duration=%.6f Timeout=%.6f)"), *ContextName, Now - StartTime, Timeout);
+			UE_LOG(LogStateGraph, Log, TEXT("[%s] State graph timed out (Duration=%.6f Timeout=%.6f)"), *ContextName, Now - StartTime, Timeout);
 			SetStatus(EStatus::TimedOut);
 			return;
 		}
@@ -353,7 +355,7 @@ void FStateGraph::Run()
 
 	bRunning = true;
 
-	UE_LOG_STATEGRAPH(Verbose, TEXT("[%s] Starting run loop (Now=%.06f)"), *ContextName, Now);
+	UE_LOG(LogStateGraph, Verbose, TEXT("[%s] Starting run loop (Now=%.06f)"), *ContextName, Now);
 
 	uint32 Blocked = 0;
 	uint32 Started = 0;
@@ -414,7 +416,7 @@ void FStateGraph::Run()
 				break;
 			}
 
-			UE_LOG_STATEGRAPH(Log, TEXT("[%s] Starting node"), *(*Node)->ContextName);
+			UE_LOG(LogStateGraph, Log, TEXT("[%s] Starting node"), *(*Node)->ContextName);
 			(*Node)->SetStatus(FStateGraphNode::EStatus::Started);
 
 			if (!StateGraphWeakPtr.IsValid())
@@ -471,7 +473,7 @@ void FStateGraph::Run()
 				if (NodeTimeout <= 0.f)
 				{
 					++TimedOut;
-					UE_LOG_STATEGRAPH(Log, TEXT("[%s] Node timed out (Duration=%.6f Timeout=%.6f)"), *(*Node)->ContextName, Now - (*Node)->StartTime, (*Node)->Timeout);
+					UE_LOG(LogStateGraph, Log, TEXT("[%s] Node timed out (Duration=%.6f Timeout=%.6f)"), *(*Node)->ContextName, Now - (*Node)->StartTime, (*Node)->Timeout);
 					(*Node)->SetStatus(FStateGraphNode::EStatus::TimedOut);
 
 					if (!StateGraphWeakPtr.IsValid())
@@ -514,20 +516,20 @@ void FStateGraph::Run()
 			break;
 
 		default:
-			UE_LOG_STATEGRAPH(Error, TEXT("[%s] Unknown state"), *(*Node)->ContextName);
+			UE_LOG(LogStateGraph, Error, TEXT("[%s] Unknown state"), *(*Node)->ContextName);
 			checkNoEntry();
 			break;
 		}
 	}
 
 	bRunning = false;
-	UE_LOG_STATEGRAPH(Verbose, TEXT("[%s] Duration=%.6f Timeout=%.6f Blocked=%d Started=%d Running=%d Completed=%d Removed=%d TimedOut=%d"),
+	UE_LOG(LogStateGraph, Verbose, TEXT("[%s] Duration=%.6f Timeout=%.6f Blocked=%d Started=%d Running=%d Completed=%d Removed=%d TimedOut=%d"),
 		*ContextName, Now - StartTime, Timeout, Blocked, Started, Running, Completed, Removed, TimedOut);
 
 	if (bRunAgain)
 	{
 		bRunAgain = false;
-		LogDebugInfo();
+		LogDebugInfo(ELogVerbosity::Type::VeryVerbose);
 		Run();
 		return;
 	}
@@ -541,19 +543,19 @@ void FStateGraph::Run()
 		if (Blocked == 0 && TimedOut == 0)
 		{
 			CompletedTime = Now;
-			UE_LOG_STATEGRAPH(Log, TEXT("[%s] Completed (Duration=%.6f Timeout=%.6f)"), *ContextName, CompletedTime - StartTime, Timeout);
+			UE_LOG(LogStateGraph, Log, TEXT("[%s] Completed (Duration=%.6f Timeout=%.6f)"), *ContextName, CompletedTime - StartTime, Timeout);
 			SetStatus(EStatus::Completed);
 		}
 		else
 		{
-			UE_LOG_STATEGRAPH(Warning, TEXT("[%s] Blocked on %d nodes, timed out %d nodes"), *ContextName, Blocked, TimedOut);
-			LogDebugInfo(true);
+			UE_LOG(LogStateGraph, Warning, TEXT("[%s] Blocked on %d nodes, timed out %d nodes"), *ContextName, Blocked, TimedOut);
+			LogDebugInfo(ELogVerbosity::Type::Warning);
 			SetStatus(EStatus::Blocked);
 		}
 	}
 	else
 	{
-		UE_LOG_STATEGRAPH(VeryVerbose, TEXT("[%s] Waiting on %d nodes"), *ContextName, Started + Running);
+		UE_LOG(LogStateGraph, VeryVerbose, TEXT("[%s] Waiting on %d nodes"), *ContextName, Started + Running);
 		SetStatus(EStatus::Waiting);
 	}
 
@@ -564,14 +566,14 @@ void FStateGraph::Run()
 
 	if ((Status == EStatus::Blocked || Status == EStatus::Waiting) && NextTimeout > 0.f)
 	{
-		UE_LOG_STATEGRAPH(Verbose, TEXT("[%s] Setting timer for %.6f"), *ContextName, NextTimeout);
+		UE_LOG(LogStateGraph, Verbose, TEXT("[%s] Setting timer for %.6f"), *ContextName, NextTimeout);
 		TimeoutTicker = FTSTicker::GetCoreTicker().AddTicker(FTickerDelegate::CreateSPLambda(this, [this](float DeltaTime) {
 			Run();
 			return false;
 		}), NextTimeout);
 	}
 
-	LogDebugInfo();
+	LogDebugInfo(ELogVerbosity::Type::VeryVerbose);
 }
 
 void FStateGraph::Reset()
@@ -579,7 +581,7 @@ void FStateGraph::Reset()
 	// Keep a reference to check if the state graph is destroyed during external functions.
 	FStateGraphWeakPtr StateGraphWeakPtr(AsWeak());
 
-	UE_LOG_STATEGRAPH(Verbose, TEXT("[%s] Resetting state graph"), *ContextName);
+	UE_LOG(LogStateGraph, Verbose, TEXT("[%s] Resetting state graph"), *ContextName);
 	StartTime = 0.f;
 	CompletedTime = 0.f;
 	bRunAgain = false;
@@ -608,27 +610,29 @@ void FStateGraph::Reset()
 
 void FStateGraph::Pause()
 {
-	UE_LOG_STATEGRAPH(Verbose, TEXT("[%s] Pausing state graph"), *ContextName);
+	UE_LOG(LogStateGraph, Verbose, TEXT("[%s] Pausing state graph"), *ContextName);
 	bRunAgain = false;
 	SetStatus(EStatus::Paused);
 }
 
-void FStateGraph::LogDebugInfo(bool bWarning)
+void FStateGraph::LogDebugInfo(ELogVerbosity::Type Verbosity)
 {
-	if (!bWarning && !UE_LOG_ACTIVE(LogStateGraph, VeryVerbose))
+	if (Verbosity > UE_GET_LOG_VERBOSITY(LogStateGraph))
 	{
 		return;
 	}
 
 	FString Message = FString::Printf(TEXT("[%s] Status=%s Nodes=%d Duration=%.6f Timeout=%.6f"), *ContextName, GetStatusName(), Nodes.Num(), GetDuration(), Timeout);
 
-	if (bWarning)
+	// Dynamic verbosity isn't supported by the UE_LOG macro, so to avoid having to handle every level, we just use Warning and Log.
+	// Fatal and Error verbosity levels are not fully supported, as they'll only show up as Warning.
+	if (Verbosity <= ELogVerbosity::Type::Warning)
 	{
-		UE_LOG_STATEGRAPH(Warning, TEXT("%s"), *Message);
+		UE_LOG(LogStateGraph, Warning, TEXT("%s"), *Message);
 	}
 	else
 	{
-		UE_LOG_STATEGRAPH(VeryVerbose, TEXT("%s"), *Message);
+		UE_LOG(LogStateGraph, Log, TEXT("%s"), *Message);
 	}
 
 	for (auto Node : Nodes)
@@ -668,13 +672,13 @@ void FStateGraph::LogDebugInfo(bool bWarning)
 		Message = FString::Printf(TEXT("[%s.%s] Status=%s Duration=%.6f Timeout=%.6f Dependencies(%s)"),
 			*ContextName, *Node.Key.ToString(), Node.Value->GetStatusName(), Node.Value->GetDuration(), Node.Value->Timeout, *FString::Join(Dependencies, TEXT(" ")));
 
-		if (bWarning)
+		if (Verbosity <= ELogVerbosity::Type::Warning)
 		{
-			UE_LOG_STATEGRAPH(Warning, TEXT("%s"), *Message);
+			UE_LOG(LogStateGraph, Warning, TEXT("%s"), *Message);
 		}
 		else
 		{
-			UE_LOG_STATEGRAPH(VeryVerbose, TEXT("%s"), *Message);
+			UE_LOG(LogStateGraph, Log, TEXT("%s"), *Message);
 		}
 	}
 }
diff --git a/Engine/Plugins/Experimental/StateGraph/Source/StateGraph/Public/StateGraph.h b/Engine/Plugins/Experimental/StateGraph/Source/StateGraph/Public/StateGraph.h
--- a/Engine/Plugins/Experimental/StateGraph/Source/StateGraph/Public/StateGraph.h
+++ b/Engine/Plugins/Experimental/StateGraph/Source/StateGraph/Public/StateGraph.h
@@ -27,7 +27,6 @@
 #include "UObject/NameTypes.h"
 
 STATEGRAPH_API DECLARE_LOG_CATEGORY_EXTERN(LogStateGraph, Log, All);
-#define UE_LOG_STATEGRAPH(Verbosity, Format, ...) UE_LOG(LogStateGraph, Verbosity, TEXT("[%s] ") Format, ANSI_TO_TCHAR(__FUNCTION__), ##__VA_ARGS__)
 
 namespace UE {
 
@@ -103,7 +102,7 @@ public:
 	virtual void Reset();
 
 	/** Get the state graph for this node, which may be invalid if this node has not been added or was removed. */
-	FStateGraphPtr GetStateGraph()
+	FStateGraphPtr GetStateGraph() const
 	{
 		return StateGraphWeakPtr.Pin();
 	}
@@ -359,7 +358,7 @@ public:
 	void Pause();
 
 	/** Log debug info for the state graph and each node. */
-	void LogDebugInfo(bool bWarning = false);
+	void LogDebugInfo(ELogVerbosity::Type Verbosity);
 
 protected:
 
@@ -391,13 +390,14 @@ protected:
 template<typename NodeType, typename... ArgsTypes>
 TSharedPtr<NodeType, ESPMode::ThreadSafe> FStateGraphNode::Next(ArgsTypes&&... Args) const
 {
-	if (!StateGraphWeakPtr.IsValid())
+	FStateGraphPtr StateGraphPtr = GetStateGraph();
+	if (!StateGraphPtr)
 	{
-		UE_LOG_STATEGRAPH(Warning, TEXT("[%s] Node attempted to be created with invalid state graph"), *ContextName);
+		UE_LOG(LogStateGraph, Warning, TEXT("[%s] Node attempted to be created with invalid state graph"), *ContextName);
 		return TSharedPtr<NodeType, ESPMode::ThreadSafe>();
 	}
 
-	TSharedPtr<NodeType, ESPMode::ThreadSafe> Node(StateGraphWeakPtr.Pin()->CreateNode<NodeType>(Forward<ArgsTypes>(Args)...));
+	TSharedPtr<NodeType, ESPMode::ThreadSafe> Node(StateGraphPtr->CreateNode<NodeType>(Forward<ArgsTypes>(Args)...));
 	Node->Dependencies.Add(Name);
 	return Node;
 }
@@ -405,13 +405,14 @@ TSharedPtr<NodeType, ESPMode::ThreadSafe> FStateGraphNode::Next(ArgsTypes&&... A
 template<typename... ArgsTypes>
 FStateGraphNodeFunctionPtr FStateGraphNode::Next(ArgsTypes&&... Args) const
 {
-	if (!StateGraphWeakPtr.IsValid())
+	FStateGraphPtr StateGraphPtr = GetStateGraph();
+	if (!StateGraphPtr)
 	{
-		UE_LOG_STATEGRAPH(Warning, TEXT("[%s] Node attempted to be created with invalid state graph"), *ContextName);
+		UE_LOG(LogStateGraph, Warning, TEXT("[%s] Node attempted to be created with invalid state graph"), *ContextName);
 		return FStateGraphNodeFunctionPtr();
 	}
 
-	FStateGraphNodeFunctionPtr Node(StateGraphWeakPtr.Pin()->CreateNode(Forward<ArgsTypes>(Args)...));
+	FStateGraphNodeFunctionPtr Node(StateGraphPtr->CreateNode(Forward<ArgsTypes>(Args)...));
 	Node->Dependencies.Add(Name);
 	return Node;
 }
