StateGraphManager and PreLoginAsync refactor:
- Add base StateGraphManager class to handle registration and creation.
- Add PreLoginAsyncManager to handle PreLoginAsync graphs.
#rnx

diff --git a/Engine/Plugins/Experimental/StateGraph/Source/StateGraphManager/Private/PreLoginAsyncManager.cpp b/Engine/Plugins/Experimental/StateGraph/Source/StateGraphManager/Private/PreLoginAsyncManager.cpp
new file mode 100644
--- /dev/null
+++ b/Engine/Plugins/Experimental/StateGraph/Source/StateGraphManager/Private/PreLoginAsyncManager.cpp
@@ -0,0 +1,60 @@
+// Copyright Epic Games, Inc. All Rights Reserved.
+
+#include "PreLoginAsyncManager.h"
+
+namespace UE::PreLoginAsync
+{
+	const FName StateGraphName("PreLoginAsync");
+	const FName OptionsName("Options");
+} // UE::PreLoginAsync
+
+#if WITH_SERVER_CODE
+
+void UPreLoginAsyncManager::InitializeStateGraph(UE::FStateGraph& StateGraph, const FString& Options, const FString& Address, const FUniqueNetIdRepl& UniqueId, const AGameModeBase::FOnPreLoginCompleteDelegate& OnComplete)
+{
+	UE::FStateGraphRef* OldStateGraph = RunningStateGraphs.Find(UniqueId);
+	if (OldStateGraph)
+	{
+		UE_LOG(LogStateGraph, Warning, TEXT("[%s] Duplicate login detected"), *StateGraph.GetContextName());
+		CompleteLogin((*OldStateGraph).Get(), TEXT("Duplicate login detected"));
+	}
+
+	RunningStateGraphs.Emplace(UniqueId, StateGraph.AsShared());
+
+	StateGraph.OnStatusChanged.AddWeakLambda(this,
+		[this](UE::FStateGraph& StateGraph, UE::FStateGraph::EStatus OldStatus, UE::FStateGraph::EStatus NewStatus)
+		{
+			if (NewStatus == UE::FStateGraph::EStatus::Completed)
+			{
+				CompleteLogin(StateGraph, FString());
+			}
+			else if (NewStatus == UE::FStateGraph::EStatus::Blocked)
+			{
+				CompleteLogin(StateGraph, TEXT("PreLoginAsync blocked"));
+			}
+			else if (NewStatus == UE::FStateGraph::EStatus::TimedOut)
+			{
+				CompleteLogin(StateGraph, TEXT("PreLoginAsync timed out"));
+			}
+		});
+
+	StateGraph.CreateNode<UE::PreLoginAsync::FOptions>(this, Options, Address, UniqueId, OnComplete);
+}
+
+void UPreLoginAsyncManager::CompleteLogin(UE::FStateGraph& StateGraph, const FString& Error)
+{
+	UE_LOG(LogStateGraph, Log, TEXT("[%s] CompleteLogin: %s"), *StateGraph.GetContextName(), Error.IsEmpty() ? TEXT("Success") : *Error);
+
+	UE::PreLoginAsync::FOptionsPtr Options = UE::PreLoginAsync::FOptions::Get(StateGraph);
+	if (Options)
+	{
+		Options->OnComplete.ExecuteIfBound(Error);
+		TObjectPtr<UPreLoginAsyncManager> Manager = Options->WeakManager.Get();
+		if (Manager)
+		{
+			Manager->RunningStateGraphs.Remove(Options->UniqueId);
+		}
+	}
+}
+
+#endif // WITH_SERVER_CODE
diff --git a/Engine/Plugins/Experimental/StateGraph/Source/StateGraphManager/Private/StateGraphManager.cpp b/Engine/Plugins/Experimental/StateGraph/Source/StateGraphManager/Private/StateGraphManager.cpp
new file mode 100644
--- /dev/null
+++ b/Engine/Plugins/Experimental/StateGraph/Source/StateGraphManager/Private/StateGraphManager.cpp
@@ -0,0 +1,49 @@
+// Copyright Epic Games, Inc. All Rights Reserved.
+
+#include "StateGraphManager.h"
+
+#include "Modules/ModuleManager.h"
+#include "StateGraph.h"
+
+class FStateGraphManagerModule : public IModuleInterface
+{
+};
+
+IMPLEMENT_MODULE(FStateGraphManagerModule, StateGraphManager);
+
+namespace UE
+{
+
+FStateGraphManager::~FStateGraphManager()
+{
+}
+
+void FStateGraphManager::AddCreateDelegate(const FStateGraphManagerCreateDelegate& Delegate)
+{
+	CreateDelegates.Add(Delegate);
+}
+
+UE::FStateGraphPtr FStateGraphManager::Create(const FString& ContextName)
+{
+	UE::FStateGraphPtr StateGraph = MakeShared<UE::FStateGraph>(GetStateGraphName(), ContextName);
+	StateGraph->Initialize();
+
+	for (int32 Index = 0; Index < CreateDelegates.Num(); Index++)
+	{
+		FStateGraphManagerCreateDelegate& Delegate = CreateDelegates[Index];
+		if (!Delegate.IsBound())
+		{
+			// Cleanup stale delegates.
+			CreateDelegates.RemoveAt(Index--);
+		}
+		else if (!Delegate.Execute(*StateGraph.Get()))
+		{
+			StateGraph.Reset();
+			break;
+		}
+	}
+
+	return StateGraph;
+}
+
+} // UE
diff --git a/Engine/Plugins/Experimental/StateGraph/Source/StateGraphManager/Public/PreLoginAsyncManager.h b/Engine/Plugins/Experimental/StateGraph/Source/StateGraphManager/Public/PreLoginAsyncManager.h
new file mode 100644
--- /dev/null
+++ b/Engine/Plugins/Experimental/StateGraph/Source/StateGraphManager/Public/PreLoginAsyncManager.h
@@ -0,0 +1,97 @@
+// Copyright Epic Games, Inc. All Rights Reserved.
+
+#pragma once
+
+/**
+ * State graph manager for AGameModeBase::PreLoginAsync.
+ */
+
+#include "Containers/Map.h"
+#include "GameFramework/GameModeBase.h"
+#include "GameFramework/OnlineReplStructs.h"
+#include "StateGraph.h"
+#include "StateGraphManager.h"
+#include "Subsystems/WorldSubsystem.h"
+
+#include "PreLoginAsyncManager.generated.h"
+
+class UPreLoginAsyncManager;
+
+namespace UE::PreLoginAsync
+{
+
+STATEGRAPHMANAGER_API extern const FName StateGraphName;
+STATEGRAPHMANAGER_API extern const FName OptionsName;
+
+#if WITH_SERVER_CODE
+
+class FOptions;
+using FOptionsPtr = TSharedPtr<FOptions, ESPMode::ThreadSafe>;
+
+/** State graph node that is used to store options that PreLoginAsync was called with. */
+class STATEGRAPHMANAGER_API FOptions : public UE::FStateGraphNode
+{
+public:
+	static FOptionsPtr Get(UE::FStateGraph& StateGraph)
+	{
+		return StateGraph.GetNode<FOptions>(OptionsName);
+	}
+
+	FOptions(TObjectPtr<UPreLoginAsyncManager> InManager, const FString& InOptions, const FString& InAddress, const FUniqueNetIdRepl& InUniqueId, const AGameModeBase::FOnPreLoginCompleteDelegate& InOnComplete) :
+		UE::FStateGraphNode(OptionsName),
+		WeakManager(InManager),
+		Options(InOptions),
+		Address(InAddress),
+		UniqueId(InUniqueId),
+		OnComplete(InOnComplete)
+	{
+	}
+
+	virtual void Start() override
+	{
+		// Nothing to do, this node just holds options.
+		Complete();
+	}
+
+	TWeakObjectPtr<UPreLoginAsyncManager> WeakManager;
+	FString Options;
+	FString Address;
+	FUniqueNetIdRepl UniqueId;
+	AGameModeBase::FOnPreLoginCompleteDelegate OnComplete;
+};
+
+#endif //WITH_SERVER_CODE
+
+} // UE::PreLoginAsync
+
+/** Subsystem manager that other modules and subsystems can depend on to register PreLoginAsync state graph delegates with. */
+UCLASS()
+class STATEGRAPHMANAGER_API UPreLoginAsyncManager : public UWorldSubsystem, public UE::FStateGraphManager
+{
+	GENERATED_BODY()
+
+public:
+	virtual FName GetStateGraphName() const override
+	{
+		return UE::PreLoginAsync::StateGraphName;
+	}
+
+#if WITH_SERVER_CODE
+	/**
+	 * Add the state graph to the map of running requests, detecting and removing conflicts if needed. This also creates the options node from the original
+	 * PreLoginAsync options. This isn't part of a Create() override so the OnComplete delegate can use the created state graph if needed.
+	 */
+	void InitializeStateGraph(UE::FStateGraph& StateGraph, const FString& Options, const FString& Address, const FUniqueNetIdRepl& UniqueId, const AGameModeBase::FOnPreLoginCompleteDelegate& OnComplete);
+
+	/**
+	 * Call the OnComplete delegate to finish the PreLoginAsync call and remove the state graph from the map of running requests. As long as no references
+	 * are kept to the shared pointer, this will free the state graph and all nodes. The state graph should no longer be used after this is called.
+	 */
+	static void CompleteLogin(UE::FStateGraph& StateGraph, const FString& Error);
+
+private:
+	/** Map of pending PreLoginAsync requests currently running. Only one per NetId is allowed. */
+	TMap<FUniqueNetIdRepl, UE::FStateGraphRef> RunningStateGraphs;
+
+#endif //WITH_SERVER_CODE
+};
diff --git a/Engine/Plugins/Experimental/StateGraph/Source/StateGraphManager/Public/StateGraphManager.h b/Engine/Plugins/Experimental/StateGraph/Source/StateGraphManager/Public/StateGraphManager.h
new file mode 100644
--- /dev/null
+++ b/Engine/Plugins/Experimental/StateGraph/Source/StateGraphManager/Public/StateGraphManager.h
@@ -0,0 +1,34 @@
+// Copyright Epic Games, Inc. All Rights Reserved.
+
+#pragma once
+
+/**
+ * This module provides managers to dynamically configure and create state graphs for various use cases.
+ * This includes a set of common managers for engine-level classes, but plugins and game code can define
+ * their own by creating subsystems that inherit from FStateGraphManager below.
+ */
+
+#include "Containers/Array.h"
+#include "Containers/UnrealString.h"
+#include "Delegates/Delegate.h"
+#include "StateGraphFwd.h"
+#include "UObject/NameTypes.h"
+
+namespace UE
+{
+
+DECLARE_DELEGATE_RetVal_OneParam(bool, FStateGraphManagerCreateDelegate, UE::FStateGraph& /*StateGraph*/);
+
+class STATEGRAPHMANAGER_API FStateGraphManager
+{
+public:
+	virtual ~FStateGraphManager();
+	virtual FName GetStateGraphName() const = 0;
+	virtual void AddCreateDelegate(const FStateGraphManagerCreateDelegate& Delegate);
+	virtual UE::FStateGraphPtr Create(const FString& ContextName = FString());
+
+private:
+	TArray<FStateGraphManagerCreateDelegate> CreateDelegates;
+};
+
+} // UE
diff --git a/Engine/Plugins/Experimental/StateGraph/Source/StateGraphManager/StateGraphManager.Build.cs b/Engine/Plugins/Experimental/StateGraph/Source/StateGraphManager/StateGraphManager.Build.cs
new file mode 100644
--- /dev/null
+++ b/Engine/Plugins/Experimental/StateGraph/Source/StateGraphManager/StateGraphManager.Build.cs
@@ -0,0 +1,19 @@
+// Copyright Epic Games, Inc. All Rights Reserved.
+
+using UnrealBuildTool;
+
+public class StateGraphManager : ModuleRules
+{
+	public StateGraphManager(ReadOnlyTargetRules Target) : base(Target)
+	{
+		PublicDependencyModuleNames.AddRange(new string[] {
+			"Core",
+			"Engine",
+			"StateGraph"
+		});
+
+		PrivateDependencyModuleNames.AddRange(new string[] {
+			"CoreUObject"
+		});
+	}
+}
diff --git a/Engine/Plugins/Experimental/StateGraph/StateGraph.uplugin b/Engine/Plugins/Experimental/StateGraph/StateGraph.uplugin
--- a/Engine/Plugins/Experimental/StateGraph/StateGraph.uplugin
+++ b/Engine/Plugins/Experimental/StateGraph/StateGraph.uplugin
@@ -13,6 +13,11 @@
 			"Name": "StateGraph",
 			"Type": "Runtime",
 			"LoadingPhase": "PostConfigInit"
+		},
+		{
+			"Name": "StateGraphManager",
+			"Type": "Runtime",
+			"LoadingPhase": "PostConfigInit"
 		}
 	]
 }
