Add RegisterServer state graph manager.
#rnx

diff --git a/Engine/Plugins/Experimental/StateGraph/Source/StateGraphManager/Private/RegisterServerManager.cpp b/Engine/Plugins/Experimental/StateGraph/Source/StateGraphManager/Private/RegisterServerManager.cpp
new file mode 100644
--- /dev/null
+++ b/Engine/Plugins/Experimental/StateGraph/Source/StateGraphManager/Private/RegisterServerManager.cpp
@@ -0,0 +1,8 @@
+// Copyright Epic Games, Inc. All Rights Reserved.
+
+#include "RegisterServerManager.h"
+
+namespace UE::RegisterServer
+{
+	const FName StateGraphName("RegisterServer");
+} // UE::RegisterServer
diff --git a/Engine/Plugins/Experimental/StateGraph/Source/StateGraphManager/Public/RegisterServerManager.h b/Engine/Plugins/Experimental/StateGraph/Source/StateGraphManager/Public/RegisterServerManager.h
new file mode 100644
--- /dev/null
+++ b/Engine/Plugins/Experimental/StateGraph/Source/StateGraphManager/Public/RegisterServerManager.h
@@ -0,0 +1,30 @@
+// Copyright Epic Games, Inc. All Rights Reserved.
+
+#pragma once
+
+/**
+ * State graph manager for AGameSession::RegisterServer.
+ */
+
+#include "StateGraphManager.h"
+#include "Subsystems/WorldSubsystem.h"
+
+#include "RegisterServerManager.generated.h"
+
+namespace UE::RegisterServer
+{
+STATEGRAPHMANAGER_API extern const FName StateGraphName;
+} // UE::RegisterServer
+
+/** Subsystem manager that other modules and subsystems can depend on to add RegisterServer state graph delegates with. */
+UCLASS()
+class STATEGRAPHMANAGER_API URegisterServerManager : public UWorldSubsystem, public UE::FStateGraphManager
+{
+	GENERATED_BODY()
+
+public:
+	virtual FName GetStateGraphName() const override
+	{
+		return UE::RegisterServer::StateGraphName;
+	}
+};
