StateGraph: Update current time after each virtual/delegate call from run, since significant time may have passed.
#rnx

diff --git a/Engine/Plugins/Experimental/StateGraph/Source/StateGraph/Private/StateGraph.cpp b/Engine/Plugins/Experimental/StateGraph/Source/StateGraph/Private/StateGraph.cpp
--- a/Engine/Plugins/Experimental/StateGraph/Source/StateGraph/Private/StateGraph.cpp
+++ b/Engine/Plugins/Experimental/StateGraph/Source/StateGraph/Private/StateGraph.cpp
@@ -316,7 +316,7 @@ void FStateGraph::RemoveAllNodes()
 
 void FStateGraph::Run()
 {
-	const double Now = FPlatformTime::Seconds();
+	double Now = FPlatformTime::Seconds();
 
 	if (Status == EStatus::NotStarted)
 	{
@@ -333,6 +333,8 @@ void FStateGraph::Run()
 		return;
 	}
 
+	Now = FPlatformTime::Seconds();
+
 	if (bRunning)
 	{
 		bRunAgain = true;
@@ -396,6 +398,7 @@ void FStateGraph::Run()
 				}
 
 				// Get node again in case CheckDependencies() removed the node.
+				Now = FPlatformTime::Seconds();
 				Node = GetNodeRef(NodeName);
 				if (Node)
 				{
@@ -406,6 +409,7 @@ void FStateGraph::Run()
 						return;
 					}
 
+					Now = FPlatformTime::Seconds();
 					++Blocked;
 				}
 				else
@@ -424,6 +428,7 @@ void FStateGraph::Run()
 				return;
 			}
 
+			Now = FPlatformTime::Seconds();
 			Node = GetNodeRef(NodeName);
 			if (!Node)
 			{
@@ -440,6 +445,7 @@ void FStateGraph::Run()
 			}
 
 			// Get node again in case Start() removed the node.
+			Now = FPlatformTime::Seconds();
 			Node = GetNodeRef(NodeName);
 			if (Node)
 			{
@@ -481,6 +487,7 @@ void FStateGraph::Run()
 						return;
 					}
 
+					Now = FPlatformTime::Seconds();
 					Node = GetNodeRef(NodeName);
 					if (Node)
 					{
@@ -489,6 +496,8 @@ void FStateGraph::Run()
 						{
 							return;
 						}
+
+						Now = FPlatformTime::Seconds();
 					}
 
 					break;
