Change RegisterServer node names to use a namespace.
#rnx

diff --git a/Engine/Plugins/Experimental/StateGraph/Source/StateGraphManager/Private/RegisterServerManager.cpp b/Engine/Plugins/Experimental/StateGraph/Source/StateGraphManager/Private/RegisterServerManager.cpp
--- a/Engine/Plugins/Experimental/StateGraph/Source/StateGraphManager/Private/RegisterServerManager.cpp
+++ b/Engine/Plugins/Experimental/StateGraph/Source/StateGraphManager/Private/RegisterServerManager.cpp
@@ -2,7 +2,7 @@
 
 #include "RegisterServerManager.h"
 
-namespace UE::RegisterServer
+namespace UE::RegisterServer::Name
 {
-	const FName StateGraphName("RegisterServer");
-} // UE::RegisterServer
+	const FName StateGraph("RegisterServer");
+} // UE::RegisterServer::Name
diff --git a/Engine/Plugins/Experimental/StateGraph/Source/StateGraphManager/Public/RegisterServerManager.h b/Engine/Plugins/Experimental/StateGraph/Source/StateGraphManager/Public/RegisterServerManager.h
--- a/Engine/Plugins/Experimental/StateGraph/Source/StateGraphManager/Public/RegisterServerManager.h
+++ b/Engine/Plugins/Experimental/StateGraph/Source/StateGraphManager/Public/RegisterServerManager.h
@@ -11,10 +11,10 @@
 
 #include "RegisterServerManager.generated.h"
 
-namespace UE::RegisterServer
+namespace UE::RegisterServer::Name
 {
-STATEGRAPHMANAGER_API extern const FName StateGraphName;
-} // UE::RegisterServer
+	STATEGRAPHMANAGER_API extern const FName StateGraph;
+} // UE::RegisterServer::Name
 
 /** Subsystem manager that other modules and subsystems can depend on to add RegisterServer state graph delegates with. */
 UCLASS()
@@ -25,6 +25,6 @@ class STATEGRAPHMANAGER_API URegisterServerManager : public UWorldSubsystem, pub
 public:
 	virtual FName GetStateGraphName() const override
 	{
-		return UE::RegisterServer::StateGraphName;
+		return UE::RegisterServer::Name::StateGraph;
 	}
 };
