Move PreLoginAsync state graph node names to a namespace.
#rnx

diff --git a/Engine/Plugins/Experimental/StateGraph/Source/StateGraphManager/Private/PreLoginAsyncManager.cpp b/Engine/Plugins/Experimental/StateGraph/Source/StateGraphManager/Private/PreLoginAsyncManager.cpp
--- a/Engine/Plugins/Experimental/StateGraph/Source/StateGraphManager/Private/PreLoginAsyncManager.cpp
+++ b/Engine/Plugins/Experimental/StateGraph/Source/StateGraphManager/Private/PreLoginAsyncManager.cpp
@@ -2,11 +2,11 @@
 
 #include "PreLoginAsyncManager.h"
 
-namespace UE::PreLoginAsync
+namespace UE::PreLoginAsync::Name
 {
-	const FName StateGraphName("PreLoginAsync");
-	const FName OptionsName("Options");
-} // UE::PreLoginAsync
+	const FName StateGraph("PreLoginAsync");
+	const FName Options("Options");
+} // UE::PreLoginAsync::Name
 
 #if WITH_SERVER_CODE
 
diff --git a/Engine/Plugins/Experimental/StateGraph/Source/StateGraphManager/Public/PreLoginAsyncManager.h b/Engine/Plugins/Experimental/StateGraph/Source/StateGraphManager/Public/PreLoginAsyncManager.h
--- a/Engine/Plugins/Experimental/StateGraph/Source/StateGraphManager/Public/PreLoginAsyncManager.h
+++ b/Engine/Plugins/Experimental/StateGraph/Source/StateGraphManager/Public/PreLoginAsyncManager.h
@@ -20,8 +20,11 @@ class UPreLoginAsyncManager;
 namespace UE::PreLoginAsync
 {
 
-STATEGRAPHMANAGER_API extern const FName StateGraphName;
-STATEGRAPHMANAGER_API extern const FName OptionsName;
+namespace Name
+{
+STATEGRAPHMANAGER_API extern const FName StateGraph;
+STATEGRAPHMANAGER_API extern const FName Options;
+} // Name
 
 #if WITH_SERVER_CODE
 
@@ -34,11 +37,11 @@ class STATEGRAPHMANAGER_API FOptions : public UE::FStateGraphNode
 public:
 	static FOptionsPtr Get(UE::FStateGraph& StateGraph)
 	{
-		return StateGraph.GetNode<FOptions>(OptionsName);
+		return StateGraph.GetNode<FOptions>(Name::Options);
 	}
 
 	FOptions(TObjectPtr<UPreLoginAsyncManager> InManager, const FString& InOptions, const FString& InAddress, const FUniqueNetIdRepl& InUniqueId, const AGameModeBase::FOnPreLoginCompleteDelegate& InOnComplete) :
-		UE::FStateGraphNode(OptionsName),
+		UE::FStateGraphNode(Name::Options),
 		WeakManager(InManager),
 		Options(InOptions),
 		Address(InAddress),
@@ -73,7 +76,7 @@ class STATEGRAPHMANAGER_API UPreLoginAsyncManager : public UWorldSubsystem, publ
 public:
 	virtual FName GetStateGraphName() const override
 	{
-		return UE::PreLoginAsync::StateGraphName;
+		return UE::PreLoginAsync::Name::StateGraph;
 	}
 
 #if WITH_SERVER_CODE
