Refactor server restart/shutdown into state graph manager.
#rnx

diff --git a/Engine/Plugins/Experimental/StateGraph/Source/StateGraph/Private/StateGraph.cpp b/Engine/Plugins/Experimental/StateGraph/Source/StateGraph/Private/StateGraph.cpp
--- a/Engine/Plugins/Experimental/StateGraph/Source/StateGraph/Private/StateGraph.cpp
+++ b/Engine/Plugins/Experimental/StateGraph/Source/StateGraph/Private/StateGraph.cpp
@@ -314,6 +314,18 @@ void FStateGraph::RemoveAllNodes()
 	}
 }
 
+bool FStateGraph::AddDependencies(FName NodeName, const TArrayView<const FName> Dependencies)
+{
+	UE::FStateGraphNodePtr Node = GetNode(NodeName);
+	if (Node)
+	{
+		Node->Dependencies.Append(Dependencies);
+		return true;
+	}
+
+	return false;
+}
+
 void FStateGraph::Run()
 {
 	double Now = FPlatformTime::Seconds();
diff --git a/Engine/Plugins/Experimental/StateGraph/Source/StateGraph/Public/StateGraph.h b/Engine/Plugins/Experimental/StateGraph/Source/StateGraph/Public/StateGraph.h
--- a/Engine/Plugins/Experimental/StateGraph/Source/StateGraph/Public/StateGraph.h
+++ b/Engine/Plugins/Experimental/StateGraph/Source/StateGraph/Public/StateGraph.h
@@ -15,6 +15,7 @@
  * - Configurable and hotfixable via ini settings.
  */
 
+#include "Containers/ArrayView.h"
 #include "Containers/Ticker.h"
 #include "Delegates/Delegate.h"
 #include "Logging/LogMacros.h"
@@ -338,6 +339,9 @@ public:
 		return CreateNode<FStateGraphNodeFunction>(NodeName, FStateGraphNodeFunctionStart::CreateLambda(Function, Forward<ArgsTypes>(Args)...));
 	}
 
+	/** Helper function for adding dependencies to a node. */
+	bool AddDependencies(FName NodeName, const TArrayView<const FName> Dependencies);
+
 	/**
 	 * Run the state graph. This is run automatically when nodes complete to try starting any blocked nodes, so this only needs to be run manually when:
 	 * - The state graph has not been run yet (Status = NotStarted).
diff --git a/Engine/Plugins/Experimental/StateGraph/Source/StateGraphManager/Private/RestartServerManager.cpp b/Engine/Plugins/Experimental/StateGraph/Source/StateGraphManager/Private/RestartServerManager.cpp
new file mode 100644
--- /dev/null
+++ b/Engine/Plugins/Experimental/StateGraph/Source/StateGraphManager/Private/RestartServerManager.cpp
@@ -0,0 +1,8 @@
+// Copyright Epic Games, Inc. All Rights Reserved.
+
+#include "RestartServerManager.h"
+
+namespace UE::RestartServer::Name
+{
+	const FName StateGraph("RestartServer");
+} // UE::RestartServer::Name
diff --git a/Engine/Plugins/Experimental/StateGraph/Source/StateGraphManager/Public/RegisterServerManager.h b/Engine/Plugins/Experimental/StateGraph/Source/StateGraphManager/Public/RegisterServerManager.h
--- a/Engine/Plugins/Experimental/StateGraph/Source/StateGraphManager/Public/RegisterServerManager.h
+++ b/Engine/Plugins/Experimental/StateGraph/Source/StateGraphManager/Public/RegisterServerManager.h
@@ -27,4 +27,12 @@ public:
 	{
 		return UE::RegisterServer::Name::StateGraph;
 	}
+
+	virtual UE::FStateGraphPtr Create(const FString& ContextName = FString()) override
+	{
+		StateGraph = UE::FStateGraphManager::Create(ContextName);
+		return StateGraph;
+	}
+
+	UE::FStateGraphPtr StateGraph;
 };
diff --git a/Engine/Plugins/Experimental/StateGraph/Source/StateGraphManager/Public/RestartServerManager.h b/Engine/Plugins/Experimental/StateGraph/Source/StateGraphManager/Public/RestartServerManager.h
new file mode 100644
--- /dev/null
+++ b/Engine/Plugins/Experimental/StateGraph/Source/StateGraphManager/Public/RestartServerManager.h
@@ -0,0 +1,38 @@
+// Copyright Epic Games, Inc. All Rights Reserved.
+
+#pragma once
+
+/**
+ * State graph manager for AGameSession::Restart.
+ */
+
+#include "StateGraphManager.h"
+#include "Subsystems/WorldSubsystem.h"
+
+#include "RestartServerManager.generated.h"
+
+namespace UE::RestartServer::Name
+{
+	STATEGRAPHMANAGER_API extern const FName StateGraph;
+} // UE::RestartServer::Name
+
+/** Subsystem manager that other modules and subsystems can depend on to add RestartServer state graph delegates with. */
+UCLASS()
+class STATEGRAPHMANAGER_API URestartServerManager : public UWorldSubsystem, public UE::FStateGraphManager
+{
+	GENERATED_BODY()
+
+public:
+	virtual FName GetStateGraphName() const override
+	{
+		return UE::RestartServer::Name::StateGraph;
+	}
+
+	virtual UE::FStateGraphPtr Create(const FString& ContextName = FString()) override
+	{
+		StateGraph = UE::FStateGraphManager::Create(ContextName);
+		return StateGraph;
+	}
+
+	UE::FStateGraphPtr StateGraph;
+};
