StateGraph updates:
- Set start time if completing a node before it was started so duration reports as 0.
- Add StateGraph function to get all nodes.
- Preserve ContextName, add a separate LogName for logging strings.
- Add tracking functions to StateGraphManager using ContextName.
- Update RegisterServer and RestartServer to use new tracking functions.
- Add ClientJoinManager for tracking game client attempts to join a server.
#rnx

diff --git a/Engine/Plugins/Experimental/StateGraph/Source/StateGraph/Private/StateGraph.cpp b/Engine/Plugins/Experimental/StateGraph/Source/StateGraph/Private/StateGraph.cpp
--- a/Engine/Plugins/Experimental/StateGraph/Source/StateGraph/Private/StateGraph.cpp
+++ b/Engine/Plugins/Experimental/StateGraph/Source/StateGraph/Private/StateGraph.cpp
@@ -22,14 +22,14 @@ namespace UE {
 
 FStateGraphNode::FStateGraphNode(FName InName) :
 	Name(InName),
-	ContextName(Name.ToString())
+	LogName(Name.ToString())
 {
-	UE_LOG(LogStateGraph, VeryVerbose, TEXT("[%s] Created node"), *ContextName);
+	UE_LOG(LogStateGraph, VeryVerbose, TEXT("[%s] Created node"), *LogName);
 }
 
 FStateGraphNode::~FStateGraphNode()
 {
-	UE_LOG(LogStateGraph, VeryVerbose, TEXT("[%s] Destroyed node"), *ContextName);
+	UE_LOG(LogStateGraph, VeryVerbose, TEXT("[%s] Destroyed node"), *LogName);
 }
 
 const TCHAR* FStateGraphNode::GetStatusName(EStatus Status)
@@ -65,7 +65,7 @@ bool FStateGraphNode::CheckDependencies() const
 	FStateGraphPtr StateGraphPtr(StateGraphWeakPtr.Pin());
 	if (!StateGraphPtr.IsValid())
 	{
-		UE_LOG(LogStateGraph, Warning, TEXT("[%s] Node checked with invalid state graph"), *ContextName);
+		UE_LOG(LogStateGraph, Warning, TEXT("[%s] Node checked with invalid state graph"), *LogName);
 		return false;
 	}
 
@@ -85,12 +85,17 @@ void FStateGraphNode::Complete()
 {
 	if (Status == EStatus::Completed)
 	{
-		UE_LOG(LogStateGraph, Warning, TEXT("[%s] Node already completed"), *ContextName);
+		UE_LOG(LogStateGraph, Warning, TEXT("[%s] Node already completed"), *LogName);
 		return;
 	}
 
 	CompletedTime = FPlatformTime::Seconds();
-	UE_LOG(LogStateGraph, Log, TEXT("[%s] Completed node (Duration=%.6f Timeout=%.6f)"), *ContextName, CompletedTime - StartTime, Timeout);
+	if (Status == EStatus::NotStarted)
+	{
+		StartTime = CompletedTime;
+	}
+
+	UE_LOG(LogStateGraph, Log, TEXT("[%s] Completed node (Duration=%.6f Timeout=%.6f)"), *LogName, CompletedTime - StartTime, Timeout);
 
 	// Keep a reference to check if the node is destroyed during external functions.
 	FStateGraphNodeWeakPtr StateGraphNodeWeakPtr(AsWeak());
@@ -113,13 +118,13 @@ void FStateGraphNode::Complete()
 	}
 	else
 	{
-		UE_LOG(LogStateGraph, Warning, TEXT("[%s] Node completed with invalid state graph"), *ContextName);
+		UE_LOG(LogStateGraph, Warning, TEXT("[%s] Node completed with invalid state graph"), *LogName);
 	}
 }
 
 void FStateGraphNode::Reset()
 {
-	UE_LOG(LogStateGraph, Verbose, TEXT("[%s] Resetting node"), *ContextName);
+	UE_LOG(LogStateGraph, Verbose, TEXT("[%s] Resetting node"), *LogName);
 	StartTime = 0.f;
 	CompletedTime = 0.f;
 	SetStatus(EStatus::NotStarted);
@@ -155,7 +160,7 @@ bool FStateGraphNodeFunction::CheckDependencies() const
 {
 	if (!StartFunction.IsBound())
 	{
-		UE_LOG(LogStateGraph, Warning, TEXT("[%s] Function node start not bound"), *GetContextName());
+		UE_LOG(LogStateGraph, Warning, TEXT("[%s] Function node start not bound"), *GetLogName());
 		return false;
 	}
 
@@ -171,7 +176,7 @@ void FStateGraphNodeFunction::Start()
 	check(StateGraph); // Must be valid since it is starting this node.
 	FStateGraphNodeWeakPtr NodeWeakPtr(AsWeak());
 	StartFunction.Execute(*StateGraph,
-		[NodeWeakPtr, ContextName = GetContextName()]()
+		[NodeWeakPtr, LogName = GetLogName()]()
 		{
 			FStateGraphNodePtr NodePtr = NodeWeakPtr.Pin();
 			if (NodePtr)
@@ -180,7 +185,7 @@ void FStateGraphNodeFunction::Start()
 			}
 			else
 			{
-				UE_LOG(LogStateGraph, Log, TEXT("[%s] Function node completed after node was destroyed"), *ContextName);
+				UE_LOG(LogStateGraph, Log, TEXT("[%s] Function node completed after node was destroyed"), *LogName);
 			}
 		});
 }
@@ -189,23 +194,24 @@ void FStateGraphNodeFunction::Start()
 
 FStateGraph::FStateGraph(FName InName, const FString& InContextName) :
 	Name(InName),
+	ContextName(InContextName),
 	ConfigSectionName(FString::Printf(TEXT("StateGraph.%s"), *InName.ToString()))
 {
-	if (InContextName.IsEmpty())
+	if (ContextName.IsEmpty())
 	{
-		ContextName = Name.ToString();
+		LogName = Name.ToString();
 	}
 	else
 	{
-		ContextName = FString::Printf(TEXT("%s(%s)"), *Name.ToString(), *InContextName);
+		LogName = FString::Printf(TEXT("%s(%s)"), *Name.ToString(), *ContextName);
 	}
 
-	UE_LOG(LogStateGraph, VeryVerbose, TEXT("[%s] Created state graph"), *ContextName);
+	UE_LOG(LogStateGraph, VeryVerbose, TEXT("[%s] Created state graph"), *LogName);
 }
 
 FStateGraph::~FStateGraph()
 {
-	UE_LOG(LogStateGraph, VeryVerbose, TEXT("[%s] Destroyed state graph"), *ContextName);
+	UE_LOG(LogStateGraph, VeryVerbose, TEXT("[%s] Destroyed state graph"), *LogName);
 	FCoreDelegates::TSOnConfigSectionsChanged().Remove(ConfigSectionsChangedDelegate);
 }
 
@@ -253,18 +259,18 @@ bool FStateGraph::AddNode(const FStateGraphNodeRef& Node)
 {
 	if (Node->StateGraphWeakPtr.IsValid())
 	{
-		UE_LOG(LogStateGraph, Warning, TEXT("[%s.%s] Node already associated with a state graph %s"), *ContextName, *Node->GetName().ToString(), *Node->ContextName);
+		UE_LOG(LogStateGraph, Warning, TEXT("[%s.%s] Node already associated with a state graph %s"), *LogName, *Node->GetName().ToString(), *Node->LogName);
 		return false;
 	}
 
 	if (Nodes.Contains(Node->GetName()))
 	{
-		UE_LOG(LogStateGraph, Warning, TEXT("[%s.%s] Node with the same name already exists"), *ContextName, *Node->GetName().ToString());
+		UE_LOG(LogStateGraph, Warning, TEXT("[%s.%s] Node with the same name already exists"), *LogName, *Node->GetName().ToString());
 		return false;
 	}
 
-	Node->ContextName = FString::Printf(TEXT("%s.%s"), *ContextName, *Node->GetName().ToString());
-	UE_LOG(LogStateGraph, Verbose, TEXT("[%s] Adding node"), *Node->ContextName);
+	Node->LogName = FString::Printf(TEXT("%s.%s"), *LogName, *Node->GetName().ToString());
+	UE_LOG(LogStateGraph, Verbose, TEXT("[%s] Adding node"), *Node->LogName);
 	Node->StateGraphWeakPtr = AsWeak();
 	Node->ConfigSectionName = FString::Printf(TEXT("%s.%s"), *ConfigSectionName, *Node->GetName().ToString());
 	Nodes.Add(Node->GetName(), Node);
@@ -277,12 +283,12 @@ bool FStateGraph::RemoveNode(FName NodeName)
 	FStateGraphNodeRef* Node = GetNodeRef(NodeName);
 	if (!Node)
 	{
-		UE_LOG(LogStateGraph, Warning, TEXT("[%s.%s] Failed to remove node"), *ContextName, *NodeName.ToString());
+		UE_LOG(LogStateGraph, Warning, TEXT("[%s.%s] Failed to remove node"), *LogName, *NodeName.ToString());
 		return false;
 	}
 
-	UE_LOG(LogStateGraph, Verbose, TEXT("[%s] Removing node"), *(*Node)->ContextName);
-	(*Node)->ContextName = (*Node)->GetName().ToString();
+	UE_LOG(LogStateGraph, Verbose, TEXT("[%s] Removing node"), *(*Node)->LogName);
+	(*Node)->LogName = (*Node)->GetName().ToString();
 	(*Node)->StateGraphWeakPtr.Reset();
 
 	// Keep a reference to check if the state graph is destroyed during external functions.
@@ -361,7 +367,7 @@ void FStateGraph::Run()
 		NextTimeout = (StartTime + Timeout) - Now;
 		if (NextTimeout <= 0.f)
 		{
-			UE_LOG(LogStateGraph, Log, TEXT("[%s] State graph timed out (Duration=%.6f Timeout=%.6f)"), *ContextName, Now - StartTime, Timeout);
+			UE_LOG(LogStateGraph, Log, TEXT("[%s] State graph timed out (Duration=%.6f Timeout=%.6f)"), *LogName, Now - StartTime, Timeout);
 			SetStatus(EStatus::TimedOut);
 			return;
 		}
@@ -369,7 +375,7 @@ void FStateGraph::Run()
 
 	bRunning = true;
 
-	UE_LOG(LogStateGraph, Verbose, TEXT("[%s] Starting run loop (Now=%.06f)"), *ContextName, Now);
+	UE_LOG(LogStateGraph, Verbose, TEXT("[%s] Starting run loop (Now=%.06f)"), *LogName, Now);
 
 	uint32 Blocked = 0;
 	uint32 Started = 0;
@@ -432,7 +438,7 @@ void FStateGraph::Run()
 				break;
 			}
 
-			UE_LOG(LogStateGraph, Log, TEXT("[%s] Starting node"), *(*Node)->ContextName);
+			UE_LOG(LogStateGraph, Log, TEXT("[%s] Starting node"), *(*Node)->LogName);
 			(*Node)->SetStatus(FStateGraphNode::EStatus::Started);
 
 			if (!StateGraphWeakPtr.IsValid())
@@ -491,7 +497,7 @@ void FStateGraph::Run()
 				if (NodeTimeout <= 0.f)
 				{
 					++TimedOut;
-					UE_LOG(LogStateGraph, Log, TEXT("[%s] Node timed out (Duration=%.6f Timeout=%.6f)"), *(*Node)->ContextName, Now - (*Node)->StartTime, (*Node)->Timeout);
+					UE_LOG(LogStateGraph, Log, TEXT("[%s] Node timed out (Duration=%.6f Timeout=%.6f)"), *(*Node)->LogName, Now - (*Node)->StartTime, (*Node)->Timeout);
 					(*Node)->SetStatus(FStateGraphNode::EStatus::TimedOut);
 
 					if (!StateGraphWeakPtr.IsValid())
@@ -537,7 +543,7 @@ void FStateGraph::Run()
 			break;
 
 		default:
-			UE_LOG(LogStateGraph, Error, TEXT("[%s] Unknown state"), *(*Node)->ContextName);
+			UE_LOG(LogStateGraph, Error, TEXT("[%s] Unknown state"), *(*Node)->LogName);
 			checkNoEntry();
 			break;
 		}
@@ -545,7 +551,7 @@ void FStateGraph::Run()
 
 	bRunning = false;
 	UE_LOG(LogStateGraph, Verbose, TEXT("[%s] Duration=%.6f Timeout=%.6f Blocked=%d Started=%d Running=%d Completed=%d Removed=%d TimedOut=%d"),
-		*ContextName, Now - StartTime, Timeout, Blocked, Started, Running, Completed, Removed, TimedOut);
+		*LogName, Now - StartTime, Timeout, Blocked, Started, Running, Completed, Removed, TimedOut);
 
 	if (bRunAgain)
 	{
@@ -564,19 +570,19 @@ void FStateGraph::Run()
 		if (Blocked == 0 && TimedOut == 0)
 		{
 			CompletedTime = Now;
-			UE_LOG(LogStateGraph, Log, TEXT("[%s] Completed (Duration=%.6f Timeout=%.6f)"), *ContextName, CompletedTime - StartTime, Timeout);
+			UE_LOG(LogStateGraph, Log, TEXT("[%s] Completed (Duration=%.6f Timeout=%.6f)"), *LogName, CompletedTime - StartTime, Timeout);
 			SetStatus(EStatus::Completed);
 		}
 		else
 		{
-			UE_LOG(LogStateGraph, Warning, TEXT("[%s] Blocked on %d nodes, timed out %d nodes"), *ContextName, Blocked, TimedOut);
+			UE_LOG(LogStateGraph, Warning, TEXT("[%s] Blocked on %d nodes, timed out %d nodes"), *LogName, Blocked, TimedOut);
 			LogDebugInfo(ELogVerbosity::Type::Warning);
 			SetStatus(EStatus::Blocked);
 		}
 	}
 	else
 	{
-		UE_LOG(LogStateGraph, VeryVerbose, TEXT("[%s] Waiting on %d nodes"), *ContextName, Started + Running);
+		UE_LOG(LogStateGraph, VeryVerbose, TEXT("[%s] Waiting on %d nodes"), *LogName, Started + Running);
 		SetStatus(EStatus::Waiting);
 	}
 
@@ -587,7 +593,7 @@ void FStateGraph::Run()
 
 	if ((Status == EStatus::Blocked || Status == EStatus::Waiting) && NextTimeout > 0.f)
 	{
-		UE_LOG(LogStateGraph, Verbose, TEXT("[%s] Setting timer for %.6f"), *ContextName, NextTimeout);
+		UE_LOG(LogStateGraph, Verbose, TEXT("[%s] Setting timer for %.6f"), *LogName, NextTimeout);
 		TimeoutTicker = FTSTicker::GetCoreTicker().AddTicker(FTickerDelegate::CreateSPLambda(this, [this](float DeltaTime) {
 			Run();
 			return false;
@@ -602,7 +608,7 @@ void FStateGraph::Reset()
 	// Keep a reference to check if the state graph is destroyed during external functions.
 	FStateGraphWeakPtr StateGraphWeakPtr(AsWeak());
 
-	UE_LOG(LogStateGraph, Verbose, TEXT("[%s] Resetting state graph"), *ContextName);
+	UE_LOG(LogStateGraph, Verbose, TEXT("[%s] Resetting state graph"), *LogName);
 	StartTime = 0.f;
 	CompletedTime = 0.f;
 	bRunAgain = false;
@@ -631,7 +637,7 @@ void FStateGraph::Reset()
 
 void FStateGraph::Pause()
 {
-	UE_LOG(LogStateGraph, Verbose, TEXT("[%s] Pausing state graph"), *ContextName);
+	UE_LOG(LogStateGraph, Verbose, TEXT("[%s] Pausing state graph"), *LogName);
 	bRunAgain = false;
 	SetStatus(EStatus::Paused);
 }
@@ -643,7 +649,7 @@ void FStateGraph::LogDebugInfo(ELogVerbosity::Type Verbosity)
 		return;
 	}
 
-	FString Message = FString::Printf(TEXT("[%s] Status=%s Nodes=%d Duration=%.6f Timeout=%.6f"), *ContextName, GetStatusName(), Nodes.Num(), GetDuration(), Timeout);
+	FString Message = FString::Printf(TEXT("[%s] Status=%s Nodes=%d Duration=%.6f Timeout=%.6f"), *LogName, GetStatusName(), Nodes.Num(), GetDuration(), Timeout);
 
 	// Dynamic verbosity isn't supported by the UE_LOG macro, so to avoid having to handle every level, we just use Warning and Log.
 	// Fatal and Error verbosity levels are not fully supported, as they'll only show up as Warning.
@@ -691,7 +697,7 @@ void FStateGraph::LogDebugInfo(ELogVerbosity::Type Verbosity)
 		}
 
 		Message = FString::Printf(TEXT("[%s.%s] Status=%s Duration=%.6f Timeout=%.6f Dependencies(%s)"),
-			*ContextName, *Node.Key.ToString(), Node.Value->GetStatusName(), Node.Value->GetDuration(), Node.Value->Timeout, *FString::Join(Dependencies, TEXT(" ")));
+			*LogName, *Node.Key.ToString(), Node.Value->GetStatusName(), Node.Value->GetDuration(), Node.Value->Timeout, *FString::Join(Dependencies, TEXT(" ")));
 
 		if (Verbosity <= ELogVerbosity::Type::Warning)
 		{
diff --git a/Engine/Plugins/Experimental/StateGraph/Source/StateGraph/Public/StateGraph.h b/Engine/Plugins/Experimental/StateGraph/Source/StateGraph/Public/StateGraph.h
--- a/Engine/Plugins/Experimental/StateGraph/Source/StateGraph/Public/StateGraph.h
+++ b/Engine/Plugins/Experimental/StateGraph/Source/StateGraph/Public/StateGraph.h
@@ -48,9 +48,9 @@ public:
 		return Name;
 	}
 
-	const FString& GetContextName() const
+	const FString& GetLogName() const
 	{
-		return ContextName;
+		return LogName;
 	}
 
 	enum class EStatus : uint8
@@ -151,7 +151,7 @@ private:
 	void SetStatus(EStatus NewStatus);
 
 	FName Name;
-	FString ContextName;
+	FString LogName;
 	FString ConfigSectionName;
 	EStatus Status = EStatus::NotStarted;
 	FStateGraphWeakPtr StateGraphWeakPtr;
@@ -208,6 +208,11 @@ public:
 		return ContextName;
 	}
 
+	const FString& GetLogName() const
+	{
+		return LogName;
+	}
+
 	enum class EStatus : uint8
 	{
 		/** State graph has not started yet. */
@@ -269,6 +274,12 @@ public:
 	/** Remove all nodes from the state graph that exist when this function is started. It's possible for nodes to be added back before this function returns if Node->Removed() adds nodes. */
 	void RemoveAllNodes();
 
+	/** Get map of all nodes in the state graph. References should not be cached in case the node is destroyed. */
+	TMap<FName, FStateGraphNodeRef>& GetNodes()
+	{
+		return Nodes;
+	}
+
 	/** Get a node from the state graph. References should not be cached in case the node is destroyed. */
 	FStateGraphNodeRef* GetNodeRef(FName NodeName)
 	{
@@ -377,6 +388,7 @@ protected:
 
 	FName Name;
 	FString ContextName;
+	FString LogName;
 	FString ConfigSectionName;
 	EStatus Status = EStatus::NotStarted;
 	double StartTime = 0.f;
@@ -397,7 +409,7 @@ TSharedPtr<NodeType, ESPMode::ThreadSafe> FStateGraphNode::Next(ArgsTypes&&... A
 	FStateGraphPtr StateGraphPtr = GetStateGraph();
 	if (!StateGraphPtr)
 	{
-		UE_LOG(LogStateGraph, Warning, TEXT("[%s] Node attempted to be created with invalid state graph"), *ContextName);
+		UE_LOG(LogStateGraph, Warning, TEXT("[%s] Node attempted to be created with invalid state graph"), *LogName);
 		return TSharedPtr<NodeType, ESPMode::ThreadSafe>();
 	}
 
@@ -412,7 +424,7 @@ FStateGraphNodeFunctionPtr FStateGraphNode::Next(ArgsTypes&&... Args) const
 	FStateGraphPtr StateGraphPtr = GetStateGraph();
 	if (!StateGraphPtr)
 	{
-		UE_LOG(LogStateGraph, Warning, TEXT("[%s] Node attempted to be created with invalid state graph"), *ContextName);
+		UE_LOG(LogStateGraph, Warning, TEXT("[%s] Node attempted to be created with invalid state graph"), *LogName);
 		return FStateGraphNodeFunctionPtr();
 	}
 
diff --git a/Engine/Plugins/Experimental/StateGraph/Source/StateGraphManager/Private/ClientJoinManager.cpp b/Engine/Plugins/Experimental/StateGraph/Source/StateGraphManager/Private/ClientJoinManager.cpp
new file mode 100644
--- /dev/null
+++ b/Engine/Plugins/Experimental/StateGraph/Source/StateGraphManager/Private/ClientJoinManager.cpp
@@ -0,0 +1,9 @@
+// Copyright Epic Games, Inc. All Rights Reserved.
+
+#include "ClientJoinManager.h"
+#include "StateGraph.h"
+
+namespace UE::ClientJoin::Name
+{
+const FName StateGraph("ClientJoin");
+} // UE::ClientJoin::Name
diff --git a/Engine/Plugins/Experimental/StateGraph/Source/StateGraphManager/Private/StateGraphManager.cpp b/Engine/Plugins/Experimental/StateGraph/Source/StateGraphManager/Private/StateGraphManager.cpp
--- a/Engine/Plugins/Experimental/StateGraph/Source/StateGraphManager/Private/StateGraphManager.cpp
+++ b/Engine/Plugins/Experimental/StateGraph/Source/StateGraphManager/Private/StateGraphManager.cpp
@@ -46,4 +46,38 @@ UE::FStateGraphPtr FStateGraphManager::Create(const FString& ContextName)
 	return StateGraph;
 }
 
+FStateGraphManagerTracked::~FStateGraphManagerTracked()
+{
+}
+
+UE::FStateGraphPtr FStateGraphManagerTracked::Create(const FString& ContextName)
+{
+	if (StateGraphs.Find(ContextName))
+	{
+		UE_LOG(LogStateGraph, Warning, TEXT("Failed to add duplicate state graph for context: %s"), *ContextName);
+		return nullptr;
+	}
+
+	return StateGraphs.Emplace(ContextName, UE::FStateGraphManager::Create(ContextName).ToSharedRef());
+}
+
+UE::FStateGraphPtr FStateGraphManagerTracked::Find(const FString& ContextName) const
+{
+	const UE::FStateGraphRef* StateGraph = StateGraphs.Find(ContextName);
+	if (StateGraph)
+	{
+		return *StateGraph;
+	}
+
+	return nullptr;
+}
+
+void FStateGraphManagerTracked::Remove(const FString& ContextName)
+{
+	if (StateGraphs.Remove(ContextName) == 0)
+	{
+		UE_LOG(LogStateGraph, Warning, TEXT("Failed to remove state graph for context: %s"), *ContextName);
+	}
+}
+
 } // UE
diff --git a/Engine/Plugins/Experimental/StateGraph/Source/StateGraphManager/Public/ClientJoinManager.h b/Engine/Plugins/Experimental/StateGraph/Source/StateGraphManager/Public/ClientJoinManager.h
new file mode 100644
--- /dev/null
+++ b/Engine/Plugins/Experimental/StateGraph/Source/StateGraphManager/Public/ClientJoinManager.h
@@ -0,0 +1,30 @@
+// Copyright Epic Games, Inc. All Rights Reserved.
+
+#pragma once
+
+/**
+ * State graph manager for a client when joining a server.
+ */
+
+#include "StateGraphManager.h"
+#include "Subsystems/GameInstanceSubsystem.h"
+
+#include "ClientJoinManager.generated.h"
+
+namespace UE::ClientJoin::Name
+{
+STATEGRAPHMANAGER_API extern const FName StateGraph;
+} // UE::ClientJoin::Name
+
+/** Subsystem manager that other modules and subsystems can depend on to add ClientJoin state graph delegates with. */
+UCLASS()
+class STATEGRAPHMANAGER_API UClientJoinManager : public UGameInstanceSubsystem, public UE::FStateGraphManagerTracked
+{
+	GENERATED_BODY()
+
+public:
+	virtual FName GetStateGraphName() const override
+	{
+		return UE::ClientJoin::Name::StateGraph;
+	}
+};
diff --git a/Engine/Plugins/Experimental/StateGraph/Source/StateGraphManager/Public/RegisterServerManager.h b/Engine/Plugins/Experimental/StateGraph/Source/StateGraphManager/Public/RegisterServerManager.h
--- a/Engine/Plugins/Experimental/StateGraph/Source/StateGraphManager/Public/RegisterServerManager.h
+++ b/Engine/Plugins/Experimental/StateGraph/Source/StateGraphManager/Public/RegisterServerManager.h
@@ -13,12 +13,12 @@
 
 namespace UE::RegisterServer::Name
 {
-	STATEGRAPHMANAGER_API extern const FName StateGraph;
+STATEGRAPHMANAGER_API extern const FName StateGraph;
 } // UE::RegisterServer::Name
 
 /** Subsystem manager that other modules and subsystems can depend on to add RegisterServer state graph delegates with. */
 UCLASS()
-class STATEGRAPHMANAGER_API URegisterServerManager : public UWorldSubsystem, public UE::FStateGraphManager
+class STATEGRAPHMANAGER_API URegisterServerManager : public UWorldSubsystem, public UE::FStateGraphManagerTracked
 {
 	GENERATED_BODY()
 
@@ -27,12 +27,4 @@ public:
 	{
 		return UE::RegisterServer::Name::StateGraph;
 	}
-
-	virtual UE::FStateGraphPtr Create(const FString& ContextName = FString()) override
-	{
-		StateGraph = UE::FStateGraphManager::Create(ContextName);
-		return StateGraph;
-	}
-
-	UE::FStateGraphPtr StateGraph;
 };
diff --git a/Engine/Plugins/Experimental/StateGraph/Source/StateGraphManager/Public/RestartServerManager.h b/Engine/Plugins/Experimental/StateGraph/Source/StateGraphManager/Public/RestartServerManager.h
--- a/Engine/Plugins/Experimental/StateGraph/Source/StateGraphManager/Public/RestartServerManager.h
+++ b/Engine/Plugins/Experimental/StateGraph/Source/StateGraphManager/Public/RestartServerManager.h
@@ -13,12 +13,12 @@
 
 namespace UE::RestartServer::Name
 {
-	STATEGRAPHMANAGER_API extern const FName StateGraph;
+STATEGRAPHMANAGER_API extern const FName StateGraph;
 } // UE::RestartServer::Name
 
 /** Subsystem manager that other modules and subsystems can depend on to add RestartServer state graph delegates with. */
 UCLASS()
-class STATEGRAPHMANAGER_API URestartServerManager : public UWorldSubsystem, public UE::FStateGraphManager
+class STATEGRAPHMANAGER_API URestartServerManager : public UWorldSubsystem, public UE::FStateGraphManagerTracked
 {
 	GENERATED_BODY()
 
@@ -27,12 +27,4 @@ public:
 	{
 		return UE::RestartServer::Name::StateGraph;
 	}
-
-	virtual UE::FStateGraphPtr Create(const FString& ContextName = FString()) override
-	{
-		StateGraph = UE::FStateGraphManager::Create(ContextName);
-		return StateGraph;
-	}
-
-	UE::FStateGraphPtr StateGraph;
 };
diff --git a/Engine/Plugins/Experimental/StateGraph/Source/StateGraphManager/Public/StateGraphManager.h b/Engine/Plugins/Experimental/StateGraph/Source/StateGraphManager/Public/StateGraphManager.h
--- a/Engine/Plugins/Experimental/StateGraph/Source/StateGraphManager/Public/StateGraphManager.h
+++ b/Engine/Plugins/Experimental/StateGraph/Source/StateGraphManager/Public/StateGraphManager.h
@@ -24,11 +24,35 @@ class STATEGRAPHMANAGER_API FStateGraphManager
 public:
 	virtual ~FStateGraphManager();
 	virtual FName GetStateGraphName() const = 0;
+
+	/** Add a delegate to be called when creating a new state graph instance. These are called in the order they were added. */
 	virtual void AddCreateDelegate(const FStateGraphManagerCreateDelegate& Delegate);
+
+	/** Create an untracked state graph instance. */
 	virtual UE::FStateGraphPtr Create(const FString& ContextName = FString());
 
-private:
+protected:
+	/** Array of registered delegates that are called in order when creating a new state graph instance. */
 	TArray<FStateGraphManagerCreateDelegate> CreateDelegates;
 };
 
+class STATEGRAPHMANAGER_API FStateGraphManagerTracked : public FStateGraphManager
+{
+public:
+	virtual ~FStateGraphManagerTracked();
+
+	/** Create a tracked state graph instance by ContextName. Only one instance per context is allowed, so the context should be a unique ID representing that instance. */
+	virtual UE::FStateGraphPtr Create(const FString& ContextName = FString()) override;
+
+	/** Find a tracked state graph instance by ContextName. */
+	virtual UE::FStateGraphPtr Find(const FString& ContextName = FString()) const;
+
+	/** Remove a tracked state graph instance by ContextName, if one exists. */
+	virtual void Remove(const FString& ContextName = FString());
+
+protected:
+	/** Map of state graphs currently tracked by context name. */
+	TMap<FString, UE::FStateGraphRef> StateGraphs;
+};
+
 } // UE
