Update StateGraph to use FDateTime and UTC time tracking.
#rnx

diff --git a/Engine/Plugins/Experimental/StateGraph/Source/StateGraph/Private/StateGraph.cpp b/Engine/Plugins/Experimental/StateGraph/Source/StateGraph/Private/StateGraph.cpp
--- a/Engine/Plugins/Experimental/StateGraph/Source/StateGraph/Private/StateGraph.cpp
+++ b/Engine/Plugins/Experimental/StateGraph/Source/StateGraph/Private/StateGraph.cpp
@@ -2,7 +2,6 @@
 
 #include "StateGraph.h"
 
-#include "HAL/PlatformTime.h"
 #include "Misc/ConfigCacheIni.h"
 #include "Misc/CoreDelegates.h"
 #include "Misc/Parse.h"
@@ -52,12 +51,12 @@ void FStateGraphNode::SetTimeout(double InTimeout)
 
 double FStateGraphNode::GetDuration() const
 {
-	if (StartTime == 0.f)
+	if (StartTime.GetTicks() == 0.f)
 	{
 		return 0.f;
 	}
 
-	return (CompletedTime == 0.f ? FPlatformTime::Seconds() : CompletedTime) - StartTime;
+	return ((CompletedTime.GetTicks() == 0.f ? FDateTime::UtcNow() : CompletedTime) - StartTime).GetTotalSeconds();
 }
 
 bool FStateGraphNode::CheckDependencies() const
@@ -89,13 +88,13 @@ void FStateGraphNode::Complete()
 		return;
 	}
 
-	CompletedTime = FPlatformTime::Seconds();
+	CompletedTime = FDateTime::UtcNow();
 	if (Status == EStatus::NotStarted)
 	{
 		StartTime = CompletedTime;
 	}
 
-	UE_LOG(LogStateGraph, Log, TEXT("[%s] Completed node (Duration=%.6f Timeout=%.6f)"), *LogName, CompletedTime - StartTime, Timeout);
+	UE_LOG(LogStateGraph, Log, TEXT("[%s] Completed node (Duration=%.3f Timeout=%.3f)"), *LogName, GetDuration(), Timeout);
 
 	// Keep a reference to check if the node is destroyed during external functions.
 	FStateGraphNodeWeakPtr StateGraphNodeWeakPtr(AsWeak());
@@ -252,7 +251,7 @@ double FStateGraph::GetDuration() const
 		return 0.f;
 	}
 
-	return (CompletedTime == 0.f ? FPlatformTime::Seconds() : CompletedTime) - StartTime;
+	return ((CompletedTime.GetTicks() == 0.f ? FDateTime::UtcNow() : CompletedTime) - StartTime).GetTotalSeconds();
 }
 
 bool FStateGraph::AddNode(const FStateGraphNodeRef& Node)
@@ -334,7 +333,7 @@ bool FStateGraph::AddDependencies(FName NodeName, const TArrayView<const FName>
 
 void FStateGraph::Run()
 {
-	double Now = FPlatformTime::Seconds();
+	FDateTime Now = FDateTime::UtcNow();
 
 	if (Status == EStatus::NotStarted)
 	{
@@ -351,7 +350,7 @@ void FStateGraph::Run()
 		return;
 	}
 
-	Now = FPlatformTime::Seconds();
+	Now = FDateTime::UtcNow();
 
 	if (bRunning)
 	{
@@ -364,10 +363,11 @@ void FStateGraph::Run()
 	double NextTimeout = 0.f;
 	if (Timeout > 0.f)
 	{
-		NextTimeout = (StartTime + Timeout) - Now;
+		const double Duration = (Now - StartTime).GetTotalSeconds();
+		NextTimeout = Timeout - Duration;
 		if (NextTimeout <= 0.f)
 		{
-			UE_LOG(LogStateGraph, Log, TEXT("[%s] State graph timed out (Duration=%.6f Timeout=%.6f)"), *LogName, Now - StartTime, Timeout);
+			UE_LOG(LogStateGraph, Log, TEXT("[%s] State graph timed out (Duration=%.3f Timeout=%.3f)"), *LogName, Duration, Timeout);
 			SetStatus(EStatus::TimedOut);
 			return;
 		}
@@ -375,7 +375,7 @@ void FStateGraph::Run()
 
 	bRunning = true;
 
-	UE_LOG(LogStateGraph, Verbose, TEXT("[%s] Starting run loop (Now=%.06f)"), *LogName, Now);
+	UE_LOG(LogStateGraph, Verbose, TEXT("[%s] Starting run loop"), *LogName);
 
 	uint32 Blocked = 0;
 	uint32 Started = 0;
@@ -416,7 +416,7 @@ void FStateGraph::Run()
 				}
 
 				// Get node again in case CheckDependencies() removed the node.
-				Now = FPlatformTime::Seconds();
+				Now = FDateTime::UtcNow();
 				Node = GetNodeRef(NodeName);
 				if (Node)
 				{
@@ -427,7 +427,7 @@ void FStateGraph::Run()
 						return;
 					}
 
-					Now = FPlatformTime::Seconds();
+					Now = FDateTime::UtcNow();
 					++Blocked;
 				}
 				else
@@ -446,7 +446,7 @@ void FStateGraph::Run()
 				return;
 			}
 
-			Now = FPlatformTime::Seconds();
+			Now = FDateTime::UtcNow();
 			Node = GetNodeRef(NodeName);
 			if (!Node)
 			{
@@ -463,7 +463,7 @@ void FStateGraph::Run()
 			}
 
 			// Get node again in case Start() removed the node.
-			Now = FPlatformTime::Seconds();
+			Now = FDateTime::UtcNow();
 			Node = GetNodeRef(NodeName);
 			if (Node)
 			{
@@ -493,11 +493,12 @@ void FStateGraph::Run()
 		case FStateGraphNode::EStatus::Started:
 			if ((*Node)->Timeout > 0.f)
 			{
-				const double NodeTimeout = ((*Node)->StartTime + (*Node)->Timeout) - Now;
+				const double NodeDuration = (Now - (*Node)->StartTime).GetTotalSeconds();
+				const double NodeTimeout = (*Node)->Timeout - NodeDuration;
 				if (NodeTimeout <= 0.f)
 				{
 					++TimedOut;
-					UE_LOG(LogStateGraph, Log, TEXT("[%s] Node timed out (Duration=%.6f Timeout=%.6f)"), *(*Node)->LogName, Now - (*Node)->StartTime, (*Node)->Timeout);
+					UE_LOG(LogStateGraph, Log, TEXT("[%s] Node timed out (Duration=%.3f Timeout=%.3f)"), *(*Node)->LogName, NodeDuration, (*Node)->Timeout);
 					(*Node)->SetStatus(FStateGraphNode::EStatus::TimedOut);
 
 					if (!StateGraphWeakPtr.IsValid())
@@ -505,7 +506,7 @@ void FStateGraph::Run()
 						return;
 					}
 
-					Now = FPlatformTime::Seconds();
+					Now = FDateTime::UtcNow();
 					Node = GetNodeRef(NodeName);
 					if (Node)
 					{
@@ -515,7 +516,7 @@ void FStateGraph::Run()
 							return;
 						}
 
-						Now = FPlatformTime::Seconds();
+						Now = FDateTime::UtcNow();
 					}
 
 					break;
@@ -550,8 +551,8 @@ void FStateGraph::Run()
 	}
 
 	bRunning = false;
-	UE_LOG(LogStateGraph, Verbose, TEXT("[%s] Duration=%.6f Timeout=%.6f Blocked=%d Started=%d Running=%d Completed=%d Removed=%d TimedOut=%d"),
-		*LogName, Now - StartTime, Timeout, Blocked, Started, Running, Completed, Removed, TimedOut);
+	UE_LOG(LogStateGraph, Verbose, TEXT("[%s] Duration=%.3f Timeout=%.3f Blocked=%d Started=%d Running=%d Completed=%d Removed=%d TimedOut=%d"),
+		*LogName, (Now - StartTime).GetTotalSeconds(), Timeout, Blocked, Started, Running, Completed, Removed, TimedOut);
 
 	if (bRunAgain)
 	{
@@ -570,7 +571,7 @@ void FStateGraph::Run()
 		if (Blocked == 0 && TimedOut == 0)
 		{
 			CompletedTime = Now;
-			UE_LOG(LogStateGraph, Log, TEXT("[%s] Completed (Duration=%.6f Timeout=%.6f)"), *LogName, CompletedTime - StartTime, Timeout);
+			UE_LOG(LogStateGraph, Log, TEXT("[%s] Completed (Duration=%.3f Timeout=%.3f)"), *LogName, GetDuration(), Timeout);
 			SetStatus(EStatus::Completed);
 		}
 		else
@@ -593,7 +594,7 @@ void FStateGraph::Run()
 
 	if ((Status == EStatus::Blocked || Status == EStatus::Waiting) && NextTimeout > 0.f)
 	{
-		UE_LOG(LogStateGraph, Verbose, TEXT("[%s] Setting timer for %.6f"), *LogName, NextTimeout);
+		UE_LOG(LogStateGraph, Verbose, TEXT("[%s] Setting timer for %.3f"), *LogName, NextTimeout);
 		TimeoutTicker = FTSTicker::GetCoreTicker().AddTicker(FTickerDelegate::CreateSPLambda(this, [this](float DeltaTime) {
 			Run();
 			return false;
@@ -649,7 +650,7 @@ void FStateGraph::LogDebugInfo(ELogVerbosity::Type Verbosity)
 		return;
 	}
 
-	FString Message = FString::Printf(TEXT("[%s] Status=%s Nodes=%d Duration=%.6f Timeout=%.6f"), *LogName, GetStatusName(), Nodes.Num(), GetDuration(), Timeout);
+	FString Message = FString::Printf(TEXT("[%s] Status=%s Nodes=%d Duration=%.3f Timeout=%.3f"), *LogName, GetStatusName(), Nodes.Num(), GetDuration(), Timeout);
 
 	// Dynamic verbosity isn't supported by the UE_LOG macro, so to avoid having to handle every level, we just use Warning and Log.
 	// Fatal and Error verbosity levels are not fully supported, as they'll only show up as Warning.
@@ -696,7 +697,7 @@ void FStateGraph::LogDebugInfo(ELogVerbosity::Type Verbosity)
 			Dependencies.Add(TEXT("None"));
 		}
 
-		Message = FString::Printf(TEXT("[%s.%s] Status=%s Duration=%.6f Timeout=%.6f Dependencies(%s)"),
+		Message = FString::Printf(TEXT("[%s.%s] Status=%s Duration=%.3f Timeout=%.3f Dependencies(%s)"),
 			*LogName, *Node.Key.ToString(), Node.Value->GetStatusName(), Node.Value->GetDuration(), Node.Value->Timeout, *FString::Join(Dependencies, TEXT(" ")));
 
 		if (Verbosity <= ELogVerbosity::Type::Warning)
diff --git a/Engine/Plugins/Experimental/StateGraph/Source/StateGraph/Public/StateGraph.h b/Engine/Plugins/Experimental/StateGraph/Source/StateGraph/Public/StateGraph.h
--- a/Engine/Plugins/Experimental/StateGraph/Source/StateGraph/Public/StateGraph.h
+++ b/Engine/Plugins/Experimental/StateGraph/Source/StateGraph/Public/StateGraph.h
@@ -20,6 +20,7 @@
 #include "Delegates/Delegate.h"
 #include "Logging/LogMacros.h"
 #include "Misc/CoreMisc.h"
+#include "Misc/DateTime.h"
 #include "StateGraphFwd.h"
 #include "Templates/IsClass.h"
 #include "Templates/IsInvocable.h"
@@ -90,6 +91,16 @@ public:
 	/** Set timeout in seconds from when Start is run until Complete is run. If this timeout is exceeded, the Status is set to TimedOut. */
 	void SetTimeout(double InTimeout);
 
+	FDateTime GetStartTime() const
+	{
+		return StartTime;
+	}
+
+	FDateTime GetCompletedTime() const
+	{
+		return CompletedTime;
+	}
+
 	/** Get duration of time that the node has taken to run. Returns 0 if not yet started. */
 	double GetDuration() const;
 
@@ -155,8 +166,8 @@ private:
 	FString ConfigSectionName;
 	EStatus Status = EStatus::NotStarted;
 	FStateGraphWeakPtr StateGraphWeakPtr;
-	double StartTime = 0.f;
-	double CompletedTime = 0.f;
+	FDateTime StartTime = 0.f;
+	FDateTime CompletedTime = 0.f;
 	double Timeout = 0.f;
 
 	friend FStateGraph;
@@ -262,6 +273,16 @@ public:
 	/** Set timeout in seconds from the initial run call (when Status changes from NotStarted to Running) to when Status changes to Completed. */
 	void SetTimeout(double InTimeout);
 
+	FDateTime GetStartTime() const
+	{
+		return StartTime;
+	}
+
+	FDateTime GetCompletedTime() const
+	{
+		return CompletedTime;
+	}
+
 	/** Get duration of time that the state graph has taken to run. Returns 0 if not yet started. */
 	double GetDuration() const;
 
@@ -391,8 +412,8 @@ protected:
 	FString LogName;
 	FString ConfigSectionName;
 	EStatus Status = EStatus::NotStarted;
-	double StartTime = 0.f;
-	double CompletedTime = 0.f;
+	FDateTime StartTime = 0.f;
+	FDateTime CompletedTime = 0.f;
 	double Timeout = 0.f;
 	TMap<FName, FStateGraphNodeRef> Nodes;
 	bool bRunning = false;
