Making FStateGraph and FStateGraphNode respect -NoTimeouts

[RN] FStateGraph and FStateGraphNode now respect -NoTimeouts
#tests Manually ran state graph tests

diff --git a/Engine/Plugins/Experimental/StateGraph/Source/StateGraph/Private/StateGraph.cpp b/Engine/Plugins/Experimental/StateGraph/Source/StateGraph/Private/StateGraph.cpp
--- a/Engine/Plugins/Experimental/StateGraph/Source/StateGraph/Private/StateGraph.cpp
+++ b/Engine/Plugins/Experimental/StateGraph/Source/StateGraph/Private/StateGraph.cpp
@@ -2,6 +2,7 @@
 
 #include "StateGraph.h"
 
+#include "Misc/CommandLine.h"
 #include "Misc/ConfigCacheIni.h"
 #include "Misc/CoreDelegates.h"
 #include "Misc/Parse.h"
@@ -46,6 +47,10 @@ const TCHAR* FStateGraphNode::GetStatusName(EStatus Status)
 
 void FStateGraphNode::SetTimeout(double InTimeout)
 {
+	if (FStateGraph::IsNoTimeouts())
+	{
+		return;
+	}
 	Timeout = InTimeout;
 }
 
@@ -131,7 +136,9 @@ void FStateGraphNode::Reset()
 
 void FStateGraphNode::UpdateConfig()
 {
-	GConfig->GetDouble(*ConfigSectionName, TEXT("Timeout"), Timeout, GEngineIni);
+	double RequestedTimeout = 0.;
+	GConfig->GetDouble(*ConfigSectionName, TEXT("Timeout"), RequestedTimeout, GEngineIni);
+	SetTimeout(RequestedTimeout);
 }
 
 void FStateGraphNode::SetStatus(EStatus NewStatus)
@@ -241,6 +248,10 @@ const TCHAR* FStateGraph::GetStatusName(EStatus Status)
 
 void FStateGraph::SetTimeout(double InTimeout)
 {
+	if (FStateGraph::IsNoTimeouts())
+	{
+		return;
+	}
 	Timeout = InTimeout;
 }
 
@@ -723,7 +734,9 @@ void FStateGraph::SetStatus(EStatus NewStatus)
 
 void FStateGraph::UpdateConfig()
 {
-	GConfig->GetDouble(*ConfigSectionName, TEXT("Timeout"), Timeout, GEngineIni);
+	double RequestedTimeout = 0.;
+	GConfig->GetDouble(*ConfigSectionName, TEXT("Timeout"), RequestedTimeout, GEngineIni);
+	SetTimeout(RequestedTimeout);
 }
 
 void FStateGraph::OnConfigSectionsChanged(const FString& IniFilename, const TSet<FString>& SectionNames)
@@ -771,4 +784,10 @@ void FStateGraph::OnConfigSectionsChanged(const FString& IniFilename, const TSet
 	}
 }
 
+bool FStateGraph::IsNoTimeouts()
+{
+	static bool NeverTimeout = FParse::Param(FCommandLine::Get(), TEXT("NoTimeouts"));
+	return NeverTimeout;
+}
+
 } // UE
diff --git a/Engine/Plugins/Experimental/StateGraph/Source/StateGraph/Public/StateGraph.h b/Engine/Plugins/Experimental/StateGraph/Source/StateGraph/Public/StateGraph.h
--- a/Engine/Plugins/Experimental/StateGraph/Source/StateGraph/Public/StateGraph.h
+++ b/Engine/Plugins/Experimental/StateGraph/Source/StateGraph/Public/StateGraph.h
@@ -400,6 +400,9 @@ public:
 
 protected:
 
+	static bool IsNoTimeouts();
+	friend class FStateGraphNode;
+
 	/** Set status and fire delegate if needed. */
 	UE_API void SetStatus(EStatus NewStatus);
 
