Complete PreLoginAsync stategraph with error if user disconnects before reaching PostLogin
Remove inventory handle if user diconnects before reaching PostLogin so inventory is closed

#Synced-CL 48981444

diff --git a/Engine/Plugins/Experimental/StateGraph/Source/StateGraphManager/Private/PreLoginAsyncManager.cpp b/Engine/Plugins/Experimental/StateGraph/Source/StateGraphManager/Private/PreLoginAsyncManager.cpp
--- a/Engine/Plugins/Experimental/StateGraph/Source/StateGraphManager/Private/PreLoginAsyncManager.cpp
+++ b/Engine/Plugins/Experimental/StateGraph/Source/StateGraphManager/Private/PreLoginAsyncManager.cpp
@@ -2,6 +2,8 @@
 
 #include "PreLoginAsyncManager.h"
 
+#include "GameDelegates.h"
+
 #include UE_INLINE_GENERATED_CPP_BY_NAME(PreLoginAsyncManager)
 
 namespace UE::PreLoginAsync::Name
@@ -12,6 +14,12 @@ namespace UE::PreLoginAsync::Name
 
 #if WITH_SERVER_CODE
 
+static TAutoConsoleVariable<bool> CVarFailPreloginOnPlayerConnectionLost(
+	TEXT("PreloginAsync.FailPreloginOnPlayerConnectionLost"),
+	true,
+	TEXT("Prelogin stategraph will complete as failed if player's connection is lost before it completes"),
+	ECVF_Default);
+
 void UPreLoginAsyncManager::InitializeStateGraph(UE::FStateGraph& StateGraph, const FString& Options, const FString& Address, const FUniqueNetIdRepl& UniqueId, const AGameModeBase::FOnPreLoginCompleteDelegate& OnComplete)
 {
 	UE::FStateGraphRef* OldStateGraph = RunningStateGraphs.Find(UniqueId);
@@ -41,6 +49,18 @@ void UPreLoginAsyncManager::InitializeStateGraph(UE::FStateGraph& StateGraph, co
 		});
 
 	StateGraph.CreateNode<UE::PreLoginAsync::FOptions>(this, Options, Address, UniqueId, OnComplete);
+	if (CVarFailPreloginOnPlayerConnectionLost.GetValueOnAnyThread())
+	{
+		FGameDelegates::Get().GetPendingConnectionLostDelegate().AddSPLambda(&StateGraph, [UniqueId, WeakStateGraph = StateGraph.AsWeak()](const FUniqueNetIdRepl& ConnectionUniqueId)
+			{
+				if (UniqueId == ConnectionUniqueId)
+				{
+					UE::FStateGraphPtr StateGraph = WeakStateGraph.Pin();
+					FGameDelegates::Get().GetPendingConnectionLostDelegate().RemoveAll(StateGraph.Get());
+					CompleteLogin(*StateGraph, TEXT("User disconnected"));
+				}
+			});
+	}
 }
 
 void UPreLoginAsyncManager::CompleteLogin(UE::FStateGraph& StateGraph, const FString& Error)
