UE: StateGraph remove timeout ticker when needed.

Login: Fix for client crash in FLoginStepContext::OnStepStarted()

#Synced-CL 50647933

diff --git a/Engine/Plugins/Experimental/StateGraph/Source/StateGraph/Private/StateGraph.cpp b/Engine/Plugins/Experimental/StateGraph/Source/StateGraph/Private/StateGraph.cpp
--- a/Engine/Plugins/Experimental/StateGraph/Source/StateGraph/Private/StateGraph.cpp
+++ b/Engine/Plugins/Experimental/StateGraph/Source/StateGraph/Private/StateGraph.cpp
@@ -624,6 +624,7 @@ void FStateGraph::Reset()
 	StartTime = 0.f;
 	CompletedTime = 0.f;
 	bRunAgain = false;
+	FTSTicker::RemoveTicker(TimeoutTicker);
 	SetStatus(EStatus::NotStarted);
 
 	if (!StateGraphWeakPtr.IsValid())
@@ -651,6 +652,7 @@ void FStateGraph::Pause()
 {
 	UE_LOG(LogStateGraph, Verbose, TEXT("[%s] Pausing state graph"), *LogName);
 	bRunAgain = false;
+	FTSTicker::RemoveTicker(TimeoutTicker);
 	SetStatus(EStatus::Paused);
 }
 
